Copying //Tasks/UE4/Dev-Online-Plugins-v3 to //UE4/Dev-Main (Source: //Tasks/UE4/Dev-Online-Plugins-v3 @ 3047294)

diff --git a/Engine/Plugins/Online/OnlineFramework/Source/Qos/Private/QosPrivatePCH.h b/Engine/Plugins/Online/OnlineFramework/Source/Qos/Private/QosPrivatePCH.h
new file mode 100644
--- /dev/null
+++ b/Engine/Plugins/Online/OnlineFramework/Source/Qos/Private/QosPrivatePCH.h
@@ -0,0 +1,11 @@
+// Copyright 1998-2016 Epic Games, Inc. All Rights Reserved.
+
+#pragma once
+
+#include "Core.h"
+#include "Engine.h"
+
+#include "OnlineSessionSettings.h"
+
+// Module includes
+#include "QosModule.h"
\ No newline at end of file
diff --git a/Engine/Plugins/Online/OnlineFramework/Source/Qos/Private/QosRegionManager.cpp b/Engine/Plugins/Online/OnlineFramework/Source/Qos/Private/QosRegionManager.cpp
new file mode 100644
--- /dev/null
+++ b/Engine/Plugins/Online/OnlineFramework/Source/Qos/Private/QosRegionManager.cpp
@@ -0,0 +1,345 @@
+// Copyright 1998-2016 Epic Games, Inc. All Rights Reserved.
+
+#include "QosPrivatePCH.h"
+#include "QosRegionManager.h"
+#include "QosInterface.h"
+#include "QosEvaluator.h"
+#include "QosModule.h"
+
+FOnlineSessionSettingsQos::FOnlineSessionSettingsQos(bool bInIsDedicated)
+{
+	NumPublicConnections = 1;
+	NumPrivateConnections = 0;
+
+	bIsLANMatch = false;
+	bShouldAdvertise = true;
+	bAllowJoinInProgress = true;
+	bAllowInvites = true;
+	bUsesPresence = false;
+	bAllowJoinViaPresence = true;
+	bAllowJoinViaPresenceFriendsOnly = false;
+
+	FString GameModeStr(GAMEMODE_QOS);
+	Set(SETTING_GAMEMODE, GameModeStr, EOnlineDataAdvertisementType::ViaOnlineService);
+	Set(SETTING_QOS, 1, EOnlineDataAdvertisementType::ViaOnlineService);
+	Set(SETTING_REGION, FQosInterface::Get()->GetRegionId(), EOnlineDataAdvertisementType::ViaOnlineService);
+	bIsDedicated = bInIsDedicated;
+}
+
+UQosRegionManager::UQosRegionManager(const FObjectInitializer& ObjectInitializer)
+	: bUseOldQosServers(false)
+	, NumTestsPerRegion(3)
+	, PingTimeout(5.0f)
+	, LastCheckTimestamp(0)
+	, Evaluator(nullptr)
+	, QosEvalResult(EQosCompletionResult::Invalid)
+{
+	check(GConfig);
+	GConfig->GetString(TEXT("Qos"), TEXT("ForceRegionId"), ForceRegionId, GGameIni);
+
+	// get a forced region id from the command line as an override
+	FParse::Value(FCommandLine::Get(), TEXT("McpRegion="), ForceRegionId);
+}
+
+void UQosRegionManager::PostReloadConfig(UProperty* PropertyThatWasLoaded)
+{
+	if (!HasAnyFlags(RF_ClassDefaultObject))
+	{
+		for (int32 RegionIdx = RegionOptions.Num() - 1; RegionIdx >= 0; RegionIdx--)
+		{
+			FQosRegionInfo& Region = RegionOptions[RegionIdx];
+
+			bool bFound = false;
+			for (FQosDatacenterInfo& Datacenter : Datacenters)
+			{
+				if (Datacenter.RegionId == Region.Region.RegionId)
+				{
+					bFound = true;
+				}
+			}
+
+			if (!bFound)
+			{
+				// Old value needs to be removed, preserve order
+				RegionOptions.RemoveAt(RegionIdx);
+			}
+		}
+
+		for (int32 MetaIdx = 0; MetaIdx < Datacenters.Num(); MetaIdx++)
+		{
+			FQosDatacenterInfo& Datacenter = Datacenters[MetaIdx];
+
+			bool bFound = false;
+			for (FQosRegionInfo& Region : RegionOptions)
+			{
+				if (Datacenter.RegionId == Region.Region.RegionId)
+				{
+					// Overwrite the metadata
+					Region.Region = Datacenter;
+					bFound = true;
+				}
+			}
+
+			if (!bFound)
+			{
+				// Add new value not in old list
+				FQosRegionInfo NewRegion(Datacenter);
+				RegionOptions.Insert(NewRegion, MetaIdx);
+			}
+		}
+
+		// Validate the current region selection
+		TrySetDefaultRegion();
+	}
+}
+
+int32 UQosRegionManager::GetMaxPingMs() const
+{
+	int32 MaxPing = -1;
+	if (GConfig->GetInt(TEXT("Qos"), TEXT("MaximumPingMs"), MaxPing, GGameIni) && MaxPing > 0)
+	{
+		return MaxPing;
+	}
+	return -1;
+}
+
+// static
+FString UQosRegionManager::GetDatacenterId()
+{
+	struct FDcidInfo
+	{
+		FDcidInfo()
+		{
+			FString OverrideDCID;
+			if (FParse::Value(FCommandLine::Get(), TEXT("DCID="), OverrideDCID))
+			{
+				// DCID specified on command line
+				DCIDString = OverrideDCID.ToUpper();
+			}
+			else
+			{
+				FString DefaultDCID;
+				check(GConfig);
+				if (GConfig->GetString(TEXT("Qos"), TEXT("DCID"), DefaultDCID, GGameIni))
+				{
+					// DCID specified in ini file
+					DCIDString = DefaultDCID.ToUpper();
+				}
+			}
+		}
+
+		FString DCIDString;
+	};
+	static FDcidInfo DCID;
+	return DCID.DCIDString;
+}
+
+void UQosRegionManager::BeginQosEvaluation(UWorld* World, const TSharedPtr<IAnalyticsProvider>& AnalyticsProvider, const FSimpleDelegate& OnComplete)
+{
+	check(World);
+
+	// no point doing the qos tests at all if we're forcing this
+	if (!ForceRegionId.IsEmpty())
+	{
+		World->GetTimerManager().SetTimerForNextTick(FTimerDelegate::CreateLambda([OnComplete]()
+		{
+			OnComplete.ExecuteIfBound();
+		}));
+		return;
+	}
+
+	// There are valid cached results, use them
+	if ((RegionOptions.Num() > 0) &&
+		(QosEvalResult == EQosCompletionResult::Success) &&
+		(FDateTime::UtcNow() - LastCheckTimestamp).GetTotalSeconds() <= 3)
+	{
+		World->GetTimerManager().SetTimerForNextTick(FTimerDelegate::CreateLambda([OnComplete]()
+		{
+			OnComplete.ExecuteIfBound();
+		}));
+		return;
+	}
+
+	// add to the completion delegate
+	OnQosEvalCompleteDelegate.Add(OnComplete);
+
+	// if we're already evaluating, simply return
+	if (Evaluator != nullptr)
+	{
+		return;
+	}
+
+	// create a new evaluator and start the process of running
+	Evaluator = NewObject<UQosEvaluator>();
+	Evaluator->SetWorld(World);
+	Evaluator->SetAnalyticsProvider(AnalyticsProvider);
+
+	FQosParams Params;
+	Params.ControllerId = 0;
+	Params.bUseOldQosServers = bUseOldQosServers;
+	Params.NumTestsPerRegion = NumTestsPerRegion;
+	Params.Timeout = PingTimeout;
+	Evaluator->FindDatacenters(Params, Datacenters, FOnQosSearchComplete::CreateUObject(this, &UQosRegionManager::OnQosEvaluationComplete));
+}
+
+void UQosRegionManager::OnQosEvaluationComplete(EQosCompletionResult Result, const TArray<FQosRegionInfo>& RegionInfo)
+{
+	// toss the evaluator
+	Evaluator = nullptr;
+	QosEvalResult = Result;
+	RegionOptions.Empty(RegionInfo.Num());
+
+	// Always capture the region information (its still correct, even if in a bad state)
+	RegionOptions = RegionInfo;
+
+	LastCheckTimestamp = FDateTime::UtcNow();
+
+	if (!SelectedRegionId.IsEmpty() && SelectedRegionId == TEXT("NONE"))
+	{
+		// Put the dev region back into the list and select it
+		ForceSelectRegion(SelectedRegionId);
+	}
+
+	// treat lack of any regions as a failure
+	if (RegionInfo.Num() <= 0)
+	{
+		QosEvalResult = EQosCompletionResult::Failure;
+	}
+
+	if (QosEvalResult == EQosCompletionResult::Success)
+	{
+		if (RegionOptions.Num() > 0)
+		{
+			TrySetDefaultRegion();
+		}
+	}
+	
+	// fire notifications
+	TArray<FSimpleDelegate> NotifyList = OnQosEvalCompleteDelegate;
+	OnQosEvalCompleteDelegate.Empty();
+	for (const auto& Callback : NotifyList)
+	{
+		Callback.ExecuteIfBound();
+	}
+}
+
+FString UQosRegionManager::GetRegionId() const
+{
+	if (!ForceRegionId.IsEmpty())
+	{
+		// we may have updated INI to bypass this process
+		return ForceRegionId;
+	}
+
+	if (QosEvalResult == EQosCompletionResult::Invalid)
+	{
+		// if we haven't run the evaluator just use the region from settings
+		// development dedicated server will come here, live services should use -mcpregion
+		return TEXT("NONE");
+	}
+
+	return SelectedRegionId;
+}
+
+const TArray<FQosRegionInfo>& UQosRegionManager::GetRegionOptions() const
+{
+	return RegionOptions;
+}
+
+void UQosRegionManager::ForceSelectRegion(const FString& InRegionId)
+{
+	QosEvalResult = EQosCompletionResult::Success;
+	ForceRegionId.Empty(); // remove any override (not typically used)
+
+	// make sure we can select this region
+	FString RegionId = InRegionId.ToUpper();
+	if (!SetSelectedRegion(RegionId))
+	{
+		// if not, add a fake entry and try again
+		FQosRegionInfo RegionInfo;
+		RegionInfo.Region.DisplayName = NSLOCTEXT("MMRegion", "Dev", "Development");
+		RegionInfo.Region.RegionId = RegionId;
+		RegionInfo.Region.bEnabled = true;
+		RegionInfo.Region.bVisible = true;
+		RegionInfo.Region.bBeta = false;
+		RegionInfo.Result = EQosCompletionResult::Success;
+		RegionInfo.AvgPingMs = 0;
+		RegionOptions.Add(RegionInfo);
+		verify(SetSelectedRegion(RegionId));
+	}
+}
+
+void UQosRegionManager::TrySetDefaultRegion()
+{
+	if (!IsRunningDedicatedServer())
+	{
+		// Try to set a default region if one hasn't already been selected
+		if (!SetSelectedRegion(GetRegionId()))
+		{
+			// try to select the lowest ping
+			int32 BestPing = INT_MAX;
+			FString BestRegionId;
+			for (const FQosRegionInfo& Region : RegionOptions)
+			{
+				if (!Region.Region.bBeta && Region.AvgPingMs < BestPing)
+				{
+					BestPing = Region.AvgPingMs;
+					BestRegionId = Region.Region.RegionId;
+				}
+			}
+			if (!SetSelectedRegion(BestRegionId))
+			{
+				UE_LOG(LogQos, Warning, TEXT("Unable to set a good region!"));
+				UE_LOG(LogQos, Warning, TEXT("Wanted to set %s, failed to fall back to %s"), *GetRegionId(), *BestRegionId);
+				DumpRegionStats();
+			}
+		}
+	}
+}
+
+bool UQosRegionManager::SetSelectedRegion(const FString& InRegionId)
+{
+	// make sure we've enumerated
+	if (QosEvalResult != EQosCompletionResult::Success)
+	{
+		// can't select region until we've enumerated them
+		return false;
+	}
+	
+	// make sure it's in the option list
+	FString RegionId = InRegionId.ToUpper();
+	for (const FQosRegionInfo& RegionInfo : RegionOptions)
+	{
+		if (RegionInfo.Region.RegionId == RegionId)
+		{
+			if (RegionInfo.IsUsable())
+			{
+				SelectedRegionId = RegionId;
+				return true;
+			}
+			else
+			{
+				return false;
+			}
+		}
+	}
+
+	// can't select a region not in the options list
+	return false;
+}
+
+void UQosRegionManager::DumpRegionStats()
+{
+	UE_LOG(LogQos, Display, TEXT("Region Info:"));
+	UE_LOG(LogQos, Display, TEXT("Current: %s "), *SelectedRegionId);
+	if (!ForceRegionId.IsEmpty())
+	{
+		UE_LOG(LogQos, Display, TEXT("Forced: %s "), *ForceRegionId);
+	}
+	
+	for (const FQosRegionInfo& Region : RegionOptions)
+	{
+		UE_LOG(LogQos, Display, TEXT("Region: %s [%s] Ping: %d"), *Region.Region.DisplayName.ToString(), *Region.Region.RegionId, Region.AvgPingMs);
+		UE_LOG(LogQos, Display, TEXT("\tEnabled: %d Visible: %d Beta: %d Result: %s"), Region.Region.bEnabled, Region.Region.bVisible, Region.Region.bBeta, ToString(Region.Result));
+	}
+}
diff --git a/Engine/Plugins/Online/OnlineFramework/Source/Qos/Public/QosInterface.h b/Engine/Plugins/Online/OnlineFramework/Source/Qos/Public/QosInterface.h
new file mode 100644
--- /dev/null
+++ b/Engine/Plugins/Online/OnlineFramework/Source/Qos/Public/QosInterface.h
@@ -0,0 +1,84 @@
+// Copyright 1998-2016 Epic Games, Inc. All Rights Reserved.
+
+#pragma once
+#include "OnlineSessionSettings.h"
+
+struct FQosRegionInfo;
+class IAnalyticsProvider;
+class UQosRegionManager;
+
+/**
+ * Main Qos interface for actions related to server quality of service
+ */
+class QOS_API FQosInterface : public TSharedFromThis<FQosInterface>, public FGCObject
+{
+public:
+
+	/**
+	 * Get the interface singleton
+	 */
+	static TSharedRef<FQosInterface> Get();
+
+	/**
+	 * Start running the async QoS evaluation
+	 */
+	void BeginQosEvaluation(UWorld* World, const TSharedPtr<IAnalyticsProvider>& AnalyticsProvider, const FSimpleDelegate& OnComplete);
+
+	/**
+	 * Get the region ID for this instance, checking ini and commandline overrides.
+	 *
+	 * Dedicated servers will have this value specified on the commandline
+	 *
+	 * Clients pull this value from the settings (or command line) and do a ping test to determine if the setting is viable.
+	 *
+	 * @return the default region identifier
+	 */
+	FString GetRegionId() const;
+
+	/**
+	 * Get the list of regions that the client can choose from (returned from search and must meet min ping requirements)
+	 *
+	 * If this list is empty, the client cannot play.
+	 */
+	const TArray<FQosRegionInfo>& GetRegionOptions() const;
+
+	/**
+	 * Try to set the selected region ID (must be present in GetRegionOptions)
+	 */
+	bool SetSelectedRegion(const FString& RegionId);
+
+	/**
+	 * Force the selected region creating a fake RegionOption if necessary
+	 */
+	void ForceSelectRegion(const FString& RegionId);
+
+	/**
+	 * Get the datacenter id for this instance, checking ini and commandline overrides
+	 * This is only relevant for dedicated servers (so they can advertise). Client does
+	 * not search on this (but may choose to prioritize results later)
+	 *
+	 * @return the default datacenter identifier
+	 */
+	static FString GetDatacenterId();
+
+	/**
+	 * Debug output for current region / datacenter information
+	 */
+	void DumpRegionStats();
+
+protected:
+
+	friend class FQosModule;
+	FQosInterface();
+	
+	bool Init();
+
+	/** FGCObject interface */
+	virtual void AddReferencedObjects(FReferenceCollector& Collector) override;
+
+private:
+
+	/** Reference to the evaluator for making datacenter determinations */
+	UQosRegionManager* RegionManager;
+};
+
diff --git a/Engine/Plugins/Online/OnlineFramework/Source/Qos/Public/QosRegionManager.h b/Engine/Plugins/Online/OnlineFramework/Source/Qos/Public/QosRegionManager.h
new file mode 100644
--- /dev/null
+++ b/Engine/Plugins/Online/OnlineFramework/Source/Qos/Public/QosRegionManager.h
@@ -0,0 +1,270 @@
+// Copyright 1998-2016 Epic Games, Inc. All Rights Reserved.
+
+#pragma once
+#include "OnlineSessionSettings.h"
+
+#include "QosRegionManager.generated.h"
+
+class IAnalyticsProvider;
+class UQosEvaluator;
+
+#define UNREACHABLE_PING 9999
+
+/** Enum for possible QoS return codes */
+UENUM()
+enum class EQosCompletionResult : uint8
+{
+	/** Incomplete, invalid result */
+	Invalid,
+	/** QoS operation was successful */
+	Success,
+	/** QoS operation ended in failure */
+	Failure,
+	/** QoS operation was canceled */
+	Canceled
+};
+
+/**
+ * Individual ping server details
+ */
+USTRUCT()
+struct FQosPingServerInfo
+{
+	GENERATED_USTRUCT_BODY()
+
+	/** Address of server */
+	UPROPERTY()
+	FString Address;
+	/** Port of server */
+	UPROPERTY()
+	int32 Port;
+};
+
+/**
+ * Metadata about datacenters that can be queried
+ */
+USTRUCT()
+struct FQosDatacenterInfo
+{
+	GENERATED_USTRUCT_BODY()
+
+	/** Localized name of the datacenter */
+	UPROPERTY()
+	FText DisplayName;
+	/** RegionId for this datacenter */
+	UPROPERTY()
+	FString RegionId;
+	/** Is this region tested */
+	UPROPERTY()
+	bool bEnabled;
+	/** Is this region visible in the UI */
+	UPROPERTY()
+	bool bVisible;
+	/** Is this region "beta" */
+	UPROPERTY()
+	bool bBeta;
+	/** Addresses of ping servers */
+	UPROPERTY()
+	TArray<FQosPingServerInfo> Servers;
+
+	FQosDatacenterInfo()
+		: RegionId()
+		, bEnabled(true)
+		, bVisible(true)
+		, bBeta(false)
+	{
+	}
+
+	bool IsUsable() const
+	{
+		return !RegionId.IsEmpty() && bEnabled && bVisible;
+	}
+};
+
+/** Runtime information about a given region */
+USTRUCT()
+struct QOS_API FQosRegionInfo
+{
+	GENERATED_USTRUCT_BODY()
+
+	/** Information about the region */
+	UPROPERTY(Transient)
+	FQosDatacenterInfo Region;
+
+	/** Success of the qos evaluation */
+	UPROPERTY(Transient)
+	EQosCompletionResult Result;
+	/** Avg ping times across all search results */
+	UPROPERTY(Transient)
+	int32 AvgPingMs;
+
+	/** Transient list of search results for a given region */
+	TArray<FOnlineSessionSearchResult> SearchResults;
+	/** Transient list of ping times for the above search results */
+	UPROPERTY(Transient)
+	TArray<int32> PingResults;
+	/** Last time this datacenter was checked */
+	UPROPERTY(Transient)
+	FDateTime LastCheckTimestamp;
+
+	FQosRegionInfo()
+		: Result(EQosCompletionResult::Invalid)
+		, AvgPingMs(UNREACHABLE_PING)
+		, LastCheckTimestamp(0)
+	{
+	}
+
+	FQosRegionInfo(const FQosDatacenterInfo& InMeta)
+		: Region(InMeta)
+		, Result(EQosCompletionResult::Invalid)
+		, AvgPingMs(UNREACHABLE_PING)
+		, LastCheckTimestamp(0)
+	{
+	}
+
+	/** @return if this region data is usable externally */
+	bool IsUsable() const
+	{
+		return Region.IsUsable();
+	}
+
+	/** reset the data to its default state */
+	void Reset()
+	{
+		// Only the transient values get reset
+		Result = EQosCompletionResult::Invalid;
+		AvgPingMs = UNREACHABLE_PING;
+		SearchResults.Empty();
+		PingResults.Empty();
+		LastCheckTimestamp = FDateTime(0);
+	}
+};
+
+/**
+ * Generic settings a server runs when hosting a simple QoS response service
+ */
+class QOS_API FOnlineSessionSettingsQos : public FOnlineSessionSettings
+{
+public:
+
+	FOnlineSessionSettingsQos(bool bInIsDedicated = true);
+	virtual ~FOnlineSessionSettingsQos() {}
+};
+
+/**
+ * Main Qos interface for actions related to server quality of service
+ */
+UCLASS(config = Engine)
+class QOS_API UQosRegionManager : public UObject
+{
+
+	GENERATED_UCLASS_BODY()
+
+public:
+
+	/**
+	 * Start running the async QoS evaluation 
+	 */
+	void BeginQosEvaluation(UWorld* World, const TSharedPtr<IAnalyticsProvider>& AnalyticsProvider, const FSimpleDelegate& OnComplete);
+
+	/**
+	 * Get the region ID for this instance, checking ini and commandline overrides.
+	 * 
+	 * Dedicated servers will have this value specified on the commandline
+	 * 
+	 * Clients pull this value from the settings (or command line) and do a ping test to determine if the setting is viable.
+	 *
+	 * @return the default region identifier
+	 */
+	FString GetRegionId() const;
+
+	/**
+	 * Get the list of regions that the client can choose from (returned from search and must meet min ping requirements)
+	 *
+	 * If this list is empty, the client cannot play.
+	 */
+	const TArray<FQosRegionInfo>& GetRegionOptions() const;
+
+	/**
+	 * Try to set the selected region ID (must be present in GetRegionOptions)
+	 */
+	bool SetSelectedRegion(const FString& RegionId);
+
+	/**
+	 * Force the selected region creating a fake RegionOption if necessary
+	 */
+	void ForceSelectRegion(const FString& RegionId);
+
+	/**
+	 * Get the datacenter id for this instance, checking ini and commandline overrides
+	 * This is only relevant for dedicated servers (so they can advertise). Client does 
+	 * not search on this (but may choose to prioritize results later)
+	 *
+	 * @return the default datacenter identifier
+	 */
+	static FString GetDatacenterId();
+
+	/**
+	 * Debug output for current region / datacenter information
+	 */
+	void DumpRegionStats();
+
+public:
+
+	/** Begin UObject interface */
+	virtual void PostReloadConfig(class UProperty* PropertyThatWasLoaded) override;
+	/** End UObject interface */
+
+private:
+
+	void OnQosEvaluationComplete(EQosCompletionResult Result, const TArray<FQosRegionInfo>& RegionInfo);
+
+	/**
+	 * Use the existing set value, or if it is currently invalid, set the next best region available
+	 */
+	void TrySetDefaultRegion();
+
+	/**
+	 * @return max allowable ping to any region and still be allowed to play
+	 */
+	int32 GetMaxPingMs() const;
+
+	/** Use old server method */
+	UPROPERTY(Config)
+	bool bUseOldQosServers;
+
+	/** Number of times to ping a given region using random sampling of available servers */
+	UPROPERTY(Config)
+	int32 NumTestsPerRegion;
+
+	/** Timeout value for each ping request */
+	UPROPERTY(Config)
+	float PingTimeout;
+
+	/** Expected datacenters metadata */
+	UPROPERTY(Config)
+	TArray<FQosDatacenterInfo> Datacenters;
+
+	UPROPERTY()
+	FDateTime LastCheckTimestamp;
+
+	/** Reference to the evaluator for making datacenter determinations (null when not active) */
+	UPROPERTY()
+	UQosEvaluator* Evaluator;
+	/** Result of the last datacenter test */
+	UPROPERTY()
+	EQosCompletionResult QosEvalResult;
+	/** Array of all known datacenters and their status */
+	UPROPERTY()
+	TArray<FQosRegionInfo> RegionOptions;
+
+	/** Value forced to be the region (development) */
+	UPROPERTY()
+	FString ForceRegionId;
+	/** Value set by the game to be the current region */
+	UPROPERTY()
+	FString SelectedRegionId;
+
+	TArray<FSimpleDelegate> OnQosEvalCompleteDelegate;
+};
+
diff --git a/Engine/Plugins/Online/OnlineSubsystem/OnlineSubsystem.uplugin b/Engine/Plugins/Online/OnlineSubsystem/OnlineSubsystem.uplugin
new file mode 100644
--- /dev/null
+++ b/Engine/Plugins/Online/OnlineSubsystem/OnlineSubsystem.uplugin
@@ -0,0 +1,21 @@
+{
+	"FileVersion" : 3,
+	
+	"FriendlyName" : "Online Subsystem",
+	"Version" : 1,
+	"VersionName" : "1.0",
+	"EngineVersion" : "4.2.0",
+	"Description" : "Shared code for interacting online subsystem implementations.",
+	"Category" : "TODO",
+	"CreatedBy" : "Epic Games, Inc.",
+	"CreatedByURL" : "http://epicgames.com",
+	"EnabledByDefault" : true,
+
+    "Modules": [
+        {
+            "Name": "OnlineSubsystem",
+            "Type": "Runtime",
+            "LoadingPhase" : "PostConfigInit"
+        }
+    ]
+}
\ No newline at end of file
diff --git a/Engine/Plugins/Online/OnlineSubsystem/Source/OnlineSubsystem.Build.cs b/Engine/Plugins/Online/OnlineSubsystem/Source/OnlineSubsystem.Build.cs
new file mode 100644
--- /dev/null
+++ b/Engine/Plugins/Online/OnlineSubsystem/Source/OnlineSubsystem.Build.cs
@@ -0,0 +1,35 @@
+// Copyright 1998-2016 Epic Games, Inc. All Rights Reserved.
+
+using UnrealBuildTool;
+
+public class OnlineSubsystem : ModuleRules
+{
+	public OnlineSubsystem(TargetInfo Target)
+	{
+		PublicDependencyModuleNames.AddRange(
+			new string[] {
+				"Json",
+			}
+		);
+
+        PublicIncludePaths.Add("OnlineSubsystem/Public/Interfaces");
+
+        PrivateIncludePaths.Add("OnlineSubsystem/Private");
+
+        Definitions.Add("ONLINESUBSYSTEM_PACKAGE=1");
+
+		PrivateDependencyModuleNames.AddRange(
+			new string[] { 
+				"Core", 
+				"CoreUObject",
+				"ImageCore",
+				"Sockets",
+				"JsonUtilities"
+			}
+		);
+
+
+	}
+
+  
+}
diff --git a/Engine/Plugins/Online/OnlineSubsystem/Source/Private/OnlineAsyncTaskManager.cpp b/Engine/Plugins/Online/OnlineSubsystem/Source/Private/OnlineAsyncTaskManager.cpp
new file mode 100644
--- /dev/null
+++ b/Engine/Plugins/Online/OnlineSubsystem/Source/Private/OnlineAsyncTaskManager.cpp
@@ -0,0 +1,335 @@
+// Copyright 1998-2016 Epic Games, Inc. All Rights Reserved.
+
+#include "OnlineSubsystemPrivatePCH.h"
+#include "OnlineAsyncTaskManager.h"
+
+int32 FOnlineAsyncTaskManager::InvocationCount = 0;
+
+#if !UE_BUILD_SHIPPING
+namespace OSSConsoleVariables
+{
+	/** Time to delay finalization of a task in the out queue */
+	TAutoConsoleVariable<float> CVarDelayAsyncTaskOutQueue(
+		TEXT("OSS.DelayAsyncTaskOutQueue"),
+		0.0f,
+		TEXT("Min total async task time\n")
+		TEXT("Time in secs"),
+		ECVF_Default);
+}
+#endif
+
+/** The default value for the polling interval when not set by config */
+#define POLLING_INTERVAL_MS 50
+
+FOnlineAsyncTaskManager::FOnlineAsyncTaskManager() :
+	ActiveTask(nullptr),
+	WorkEvent(nullptr),
+	PollingInterval(POLLING_INTERVAL_MS),
+	bRequestingExit(0),
+	OnlineThreadId(0)
+{
+}
+
+bool FOnlineAsyncTaskManager::Init(void)
+{
+	WorkEvent = FPlatformProcess::GetSynchEventFromPool();
+	int32 PollingConfig = POLLING_INTERVAL_MS;
+	// Read the polling interval to use from the INI file
+	if (GConfig->GetInt(TEXT("OnlineSubsystem"), TEXT("PollingIntervalInMs"), PollingConfig, GEngineIni))
+	{
+		PollingInterval = (uint32)PollingConfig;
+	}
+
+	return WorkEvent != nullptr;
+}
+
+uint32 FOnlineAsyncTaskManager::Run(void)
+{
+	InvocationCount++;
+	// This should not be set yet
+	check(OnlineThreadId == 0);
+	FPlatformAtomics::InterlockedExchange((volatile int32*)&OnlineThreadId, FPlatformTLS::GetCurrentThreadId());
+
+	do 
+	{
+		FOnlineAsyncTask* Task = NULL;
+
+		// Wait for a trigger event to start work
+		WorkEvent->Wait(PollingInterval);
+		if (!bRequestingExit)
+		{
+			SCOPE_CYCLE_COUNTER(STAT_Online_Async);
+			
+			double TimeElapsed = 0;
+
+			// Chance for services to do work
+			OnlineTick();
+
+			TArray<FOnlineAsyncTask*> CopyParallelTasks;
+
+			// Grab a copy of the parallel list
+			{
+				FScopeLock LockParallelTasks(&ParallelTasksLock);
+
+				CopyParallelTasks = ParallelTasks;
+			}
+
+			// Tick all the parallel tasks
+			for ( auto it = CopyParallelTasks.CreateIterator(); it; ++it )
+			{
+				Task = *it;
+
+				Task->Tick();
+
+				if (Task->IsDone())
+				{
+					if (Task->WasSuccessful())
+					{
+						UE_LOG(LogOnline, Verbose, TEXT("Async task '%s' succeeded in %f seconds (Parallel)"),
+							*Task->ToString(),
+							Task->GetElapsedTime());
+					}
+					else
+					{
+						UE_LOG(LogOnline, Log, TEXT("Async task '%s' failed in %f seconds (Parallel)"),
+							*Task->ToString(),
+							Task->GetElapsedTime());
+					}
+
+					// Task is done, remove from the incoming queue and add to the outgoing queue
+					RemoveFromParallelTasks( Task );
+
+					AddToOutQueue( Task );
+				}
+			}
+
+			// Now process the serial "In" queue
+			{
+				FScopeLock LockActiveTask(&ActiveTaskLock);
+				Task = ActiveTask;
+			}
+			if (Task)
+			{
+				Task->Tick();
+	
+				if (Task->IsDone())
+				{
+					if (Task->WasSuccessful())
+					{
+						UE_LOG(LogOnline, Verbose, TEXT("Async task '%s' succeeded in %f seconds"),
+							*Task->ToString(),
+							Task->GetElapsedTime());
+					}
+					else
+					{
+						UE_LOG(LogOnline, Warning, TEXT("Async task '%s' failed in %f seconds"),
+							*Task->ToString(),
+							Task->GetElapsedTime());
+					}
+						
+					// Task is done, add to the outgoing queue
+					AddToOutQueue(Task);
+					
+					{
+						FScopeLock LockActiveTask(&ActiveTaskLock);
+						ActiveTask = nullptr;
+					}
+				}
+			}
+		}
+	} 
+	while (!bRequestingExit);
+
+	return 0;
+}
+
+void FOnlineAsyncTaskManager::Stop(void)
+{
+	WorkEvent->Trigger();
+	FPlatformAtomics::InterlockedExchange(&bRequestingExit, 1);
+}
+
+void FOnlineAsyncTaskManager::Exit(void)
+{
+	FPlatformProcess::ReturnSynchEventToPool(WorkEvent);
+	WorkEvent = nullptr;
+
+	OnlineThreadId = 0;
+	InvocationCount--;
+}
+
+void FOnlineAsyncTaskManager::AddToInQueue(FOnlineAsyncTask* NewTask)
+{
+	FScopeLock Lock(&InQueueLock);
+	InQueue.Add(NewTask);
+}
+
+void FOnlineAsyncTaskManager::AddToOutQueue(FOnlineAsyncItem* CompletedItem)
+{
+	FScopeLock Lock(&OutQueueLock);
+	OutQueue.Add(CompletedItem);
+}
+
+void FOnlineAsyncTaskManager::AddToParallelTasks(FOnlineAsyncTask* NewTask)
+{
+	NewTask->Initialize();
+
+	FScopeLock LockParallelTasks(&ParallelTasksLock);
+
+	ParallelTasks.Add( NewTask );
+}
+
+void FOnlineAsyncTaskManager::RemoveFromParallelTasks(FOnlineAsyncTask* OldTask)
+{
+	FScopeLock LockParallelTasks(&ParallelTasksLock);
+
+	ParallelTasks.Remove( OldTask );
+}
+
+void FOnlineAsyncTaskManager::GameTick()
+{
+	// assert if not game thread
+	check(IsInGameThread());
+
+	FOnlineAsyncItem* Item = nullptr;
+	int32 CurrentQueueSize = 0;
+
+#if !UE_BUILD_SHIPPING
+	const float TimeToWait = OSSConsoleVariables::CVarDelayAsyncTaskOutQueue.GetValueOnGameThread();
+#endif
+
+	do 
+	{
+		Item = nullptr;
+		// Grab a completed task from the queue
+		{
+			FScopeLock LockOutQueue(&OutQueueLock);
+			CurrentQueueSize = OutQueue.Num();
+			if (CurrentQueueSize > 0)
+			{
+				Item = OutQueue[0];
+
+#if !UE_BUILD_SHIPPING
+				if (Item && Item->GetElapsedTime() >= TimeToWait)
+				{
+					OutQueue.RemoveAt(0);
+				}
+				else
+				{
+					Item = nullptr;
+				}
+#else
+				OutQueue.RemoveAt(0);	
+#endif
+			}
+		}
+
+		if (Item)
+		{
+#if !UE_BUILD_SHIPPING
+			if (TimeToWait > 0.0f)
+			{
+				UE_LOG(LogOnline, Verbose, TEXT("Async task '%s' finalizing after %f seconds"),
+					*Item->ToString(),
+					Item->GetElapsedTime());
+			}
+#endif
+
+			// Finish work and trigger delegates
+			Item->Finalize();
+			Item->TriggerDelegates();
+			delete Item;
+			Item = nullptr;
+		}
+	}
+	while (Item != NULL);
+
+	int32 QueueSize = 0;
+	bool bHasActiveTask = false;
+	{
+		{
+			FScopeLock LockInQueue(&InQueueLock);
+			QueueSize = InQueue.Num();
+		}
+		{
+			FScopeLock LockActiveTask(&ActiveTaskLock);
+			if (ActiveTask != nullptr)
+			{
+				++QueueSize;
+				bHasActiveTask = true;
+			}
+		}
+
+		if (!bHasActiveTask && QueueSize > 0)
+		{
+			// Grab the current task from the queue
+			FOnlineAsyncTask* Task = nullptr;
+			{
+				FScopeLock LockInQueue(&InQueueLock);
+				Task = InQueue[0];
+				InQueue.RemoveAt(0);
+			}
+			{
+				FScopeLock LockActiveTask(&ActiveTaskLock);
+				ActiveTask = Task;
+				ActiveTask->Initialize();
+			}
+			WorkEvent->Trigger();
+		}
+	}
+
+	SET_DWORD_STAT(STAT_Online_AsyncTasks, QueueSize);
+}
+
+void FOnlineAsyncTaskManager::Tick()
+{
+	// Tick Online services ( possibly callbacks ). 
+	OnlineTick();
+
+	{
+		// Tick all the parallel tasks - Tick unrelated tasks together. 
+
+		// Create a copy of existing tasks. 
+		TArray<FOnlineAsyncTask*> CopyParallelTasks = ParallelTasks;
+
+		// parallel Q.
+		FOnlineAsyncTask* Task = NULL;
+
+		// Iterate. 
+		for (auto it = CopyParallelTasks.CreateIterator(); it; ++it)
+		{
+			Task = *it;
+			Task->Tick();
+
+			if (Task->IsDone())
+			{
+				UE_LOG(LogOnline, Log, TEXT("Async parallel Task '%s' completed in %f seconds with %d (Parallel)"),
+					*Task->ToString(),
+					Task->GetElapsedTime(),
+					Task->WasSuccessful());
+
+				// Task is done, fixup the original parallel task queue. 
+				RemoveFromParallelTasks(Task);
+				AddToOutQueue(Task);
+			}
+		}
+	}
+
+	// Serial Q.
+	if (ActiveTask == nullptr)
+		return;
+
+	ActiveTask->Tick();
+
+	if (ActiveTask->IsDone())
+	{
+		UE_LOG(LogOnline, Log, TEXT("Async serial task '%s' completed in %f seconds with %d"),
+			*ActiveTask->ToString(),
+			ActiveTask->GetElapsedTime(),
+			ActiveTask->WasSuccessful());
+
+		// Task is done, add to the outgoing queue
+		AddToOutQueue(ActiveTask);
+		ActiveTask = nullptr;
+	}
+}
diff --git a/Engine/Plugins/Online/OnlineSubsystem/Source/Private/OnlineError.cpp b/Engine/Plugins/Online/OnlineSubsystem/Source/Private/OnlineError.cpp
new file mode 100644
--- /dev/null
+++ b/Engine/Plugins/Online/OnlineSubsystem/Source/Private/OnlineError.cpp
@@ -0,0 +1,103 @@
+// Copyright 1998-2016 Epic Games, Inc. All Rights Reserved.
+
+#include "OnlineSubsystemPrivatePCH.h"
+#include "OnlineError.h"
+#include "OnlineSubsystemTypes.h"
+
+//static
+const FString FOnlineError::GenericErrorCode = TEXT("GenericError");
+
+FOnlineError::FOnlineError()
+	: bSucceeded(false)
+	, HttpResult(0)
+{
+}
+
+FOnlineError::FOnlineError(bool bSucceededIn)
+	: bSucceeded(bSucceededIn)
+	, HttpResult(0)
+{
+}
+
+FOnlineError::FOnlineError(const FString& ErrorCodeIn)
+	: bSucceeded(false)
+	, HttpResult(0)
+{
+	SetFromErrorCode(ErrorCodeIn);
+}
+
+void FOnlineError::SetFromErrorCode(const FString& ErrorCodeIn)
+{
+	ErrorCode = ErrorCodeIn;
+	ErrorRaw = ErrorCodeIn;
+}
+
+FOnlineError::FOnlineError(const FText& ErrorMessageIn)
+	: bSucceeded(false)
+	, HttpResult(0)
+{
+	SetFromErrorMessage(ErrorMessageIn);
+}
+
+void FOnlineError::SetFromErrorMessage(const FText& ErrorMessageIn)
+{
+	ErrorMessage = ErrorMessageIn;
+	ErrorCode = FTextInspector::GetKey(ErrorMessageIn).Get(GenericErrorCode);
+	ErrorRaw = ErrorMessageIn.ToString();
+}
+
+const TCHAR* FOnlineError::ToLogString() const
+{
+	if (!ErrorMessage.IsEmpty())
+	{
+		return *ErrorMessage.ToString();
+	}
+	else if (!ErrorCode.IsEmpty())
+	{
+		return *ErrorCode;
+	}
+	else if (!ErrorRaw.IsEmpty())
+	{
+		return *ErrorRaw;
+	}
+	else
+	{
+		return TEXT("(Empty FOnlineError)");
+	}
+}
+
+EOnlineServerConnectionStatus::Type FOnlineError::GetConnectionStatusFromHttpResult() const
+{
+	if (bSucceeded)
+	{
+		return EOnlineServerConnectionStatus::Connected;
+	}
+
+	switch (HttpResult)
+	{
+	case 0:
+		// no response means we couldn't even connect
+		return EOnlineServerConnectionStatus::ConnectionDropped;
+	case 401:
+		// No auth at all
+		return EOnlineServerConnectionStatus::InvalidUser;
+	case 403:
+		// Auth failure
+		return EOnlineServerConnectionStatus::NotAuthorized;
+	case 503:
+	case 504:
+		// service is not avail
+		return EOnlineServerConnectionStatus::ServersTooBusy;
+	case 505:
+		// update to supported version
+		return EOnlineServerConnectionStatus::UpdateRequired;
+	case 408:
+	case 501:
+	case 502:
+		// other bad responses (ELB, etc)
+		return EOnlineServerConnectionStatus::ServiceUnavailable;
+	default:
+		// this is a domain specific error like a 400 or 404, this will be handled by application code
+		return EOnlineServerConnectionStatus::Connected;
+	}
+}
diff --git a/Engine/Plugins/Online/OnlineSubsystem/Source/Private/OnlineMessage.cpp b/Engine/Plugins/Online/OnlineSubsystem/Source/Private/OnlineMessage.cpp
new file mode 100644
--- /dev/null
+++ b/Engine/Plugins/Online/OnlineSubsystem/Source/Private/OnlineMessage.cpp
@@ -0,0 +1,98 @@
+// Copyright 1998-2016 Epic Games, Inc. All Rights Reserved.
+
+#include "OnlineSubsystemPrivatePCH.h"
+#include "OnlineMessageInterface.h"
+#include "NboSerializer.h"
+
+void FOnlineMessagePayload::ToBytes(TArray<uint8>& OutBytes) const
+{
+	FNboSerializeToBuffer Ar(MaxPayloadSize);
+	Ar << KeyValData;
+	Ar.TrimBuffer();
+	OutBytes = Ar.GetBuffer();
+}
+
+void FOnlineMessagePayload::FromBytes(const TArray<uint8>& InBytes)
+{
+	FNboSerializeFromBuffer Ar(InBytes.GetData(), InBytes.Num());
+	Ar >> KeyValData;
+}
+
+void FOnlineMessagePayload::ToJson(FJsonObject& OutJsonObject) const
+{
+	TArray<TSharedPtr<FJsonValue> > JsonProperties;
+	for (const auto It : KeyValData)
+	{
+		const FString& PropertyName = It.Key;
+		const FVariantData& PropertyValue = It.Value;
+
+		TSharedRef<FJsonObject> PropertyJson = PropertyValue.ToJson();
+		PropertyJson->SetStringField(TEXT("Name"), *PropertyName);
+		JsonProperties.Add(MakeShareable(new FJsonValueObject(PropertyJson)));
+	}
+	OutJsonObject.SetArrayField(TEXT("Properties"), JsonProperties);
+}
+
+FString FOnlineMessagePayload::ToJsonStr() const
+{
+	FString PayloadJsonStr;
+	TSharedRef<FJsonObject> JsonObject(new FJsonObject());
+	ToJson(*JsonObject);
+	auto JsonWriter = TJsonWriterFactory<TCHAR, TCondensedJsonPrintPolicy<TCHAR> >::Create(&PayloadJsonStr);
+	FJsonSerializer::Serialize(JsonObject, JsonWriter);
+	JsonWriter->Close();
+	return PayloadJsonStr;
+}
+
+void FOnlineMessagePayload::FromJson(const FJsonObject& JsonObject)
+{
+	if (JsonObject.HasTypedField<EJson::Array>(TEXT("Properties")))
+	{
+		KeyValData.Empty();
+		const TArray<TSharedPtr<FJsonValue> >& JsonProperties = JsonObject.GetArrayField(TEXT("Properties"));
+		for (auto JsonPropertyValue : JsonProperties)
+		{
+			TSharedPtr<FJsonObject> JsonPropertyObject = JsonPropertyValue->AsObject();
+			if (JsonPropertyObject.IsValid())
+			{
+				FString PropertyName;
+				if (JsonPropertyObject->TryGetStringField(TEXT("Name"), PropertyName) &&
+					!PropertyName.IsEmpty())
+				{
+					FVariantData PropertyData;
+					if (PropertyData.FromJson(JsonPropertyObject.ToSharedRef()))
+					{
+						KeyValData.Add(PropertyName, PropertyData);
+					}
+				}
+			}
+		}
+	}
+}
+
+void FOnlineMessagePayload::FromJsonStr(const FString& JsonStr)
+{
+	TSharedPtr<FJsonObject> JsonObject;
+	TSharedRef<TJsonReader<> > JsonReader = TJsonReaderFactory<>::Create(JsonStr);
+	if (FJsonSerializer::Deserialize(JsonReader, JsonObject) &&
+		JsonObject.IsValid())
+	{
+		FromJson(*JsonObject);
+	}
+}
+
+bool FOnlineMessagePayload::GetAttribute(const FString& AttrName, FVariantData& OutAttrValue) const
+{
+	const FVariantData* Value = KeyValData.Find(AttrName);
+	if (Value != NULL)
+	{
+		OutAttrValue = *Value;
+		return true;
+	}
+	return false;
+}
+
+void FOnlineMessagePayload::SetAttribute(const FString& AttrName, const FVariantData& AttrValue)
+{
+	KeyValData.Add(AttrName, AttrValue);
+}
\ No newline at end of file
diff --git a/Engine/Plugins/Online/OnlineSubsystem/Source/Private/OnlinePartyInterface.cpp b/Engine/Plugins/Online/OnlineSubsystem/Source/Private/OnlinePartyInterface.cpp
new file mode 100644
--- /dev/null
+++ b/Engine/Plugins/Online/OnlineSubsystem/Source/Private/OnlinePartyInterface.cpp
@@ -0,0 +1,96 @@
+// Copyright 1998-2016 Epic Games, Inc. All Rights Reserved.
+
+#include "OnlineSubsystemPrivatePCH.h"
+#include "OnlinePartyInterface.h"
+
+DEFINE_LOG_CATEGORY(LogOnlineParty);
+
+void FOnlinePartyData::ToJsonFull(FString& JsonString) const
+{
+	JsonString.Empty();
+
+	// iterate over key/val attrs and convert each entry to a json string
+	TSharedRef<FJsonObject> JsonObject(new FJsonObject());
+	TArray<TSharedPtr<FJsonValue> > JsonProperties;
+	for (auto Iterator : KeyValAttrs)
+	{
+		const FString& PropertyName = Iterator.Key;
+		const FVariantData& PropertyValue = Iterator.Value;
+
+		TSharedRef<FJsonObject> PropertyJson = PropertyValue.ToJson();
+		PropertyJson->SetStringField(TEXT("Name"), *PropertyName);
+		JsonProperties.Add(MakeShareable(new FJsonValueObject(PropertyJson)));
+	}
+	JsonObject->SetNumberField(TEXT("Rev"), RevisionCount);
+	JsonObject->SetArrayField(TEXT("Attrs"), JsonProperties);
+
+	auto JsonWriter = TJsonWriterFactory<TCHAR, TCondensedJsonPrintPolicy<TCHAR> >::Create(&JsonString);
+	FJsonSerializer::Serialize(JsonObject, JsonWriter);
+	JsonWriter->Close();
+}
+
+void FOnlinePartyData::ToJsonDirty(FString& JsonString) const
+{
+	JsonString.Empty();
+
+	// iterate over key/val attrs and convert each entry to a json string
+	TSharedRef<FJsonObject> JsonObject(new FJsonObject());
+	TArray<TSharedPtr<FJsonValue> > JsonProperties;
+	for (auto& PropertyName : DirtyKeys)
+	{
+		const FVariantData* PropertyValue = KeyValAttrs.Find(PropertyName);
+		check(PropertyValue);
+
+		TSharedRef<FJsonObject> PropertyJson = PropertyValue->ToJson();
+		PropertyJson->SetStringField(TEXT("Name"), *PropertyName);
+		JsonProperties.Add(MakeShareable(new FJsonValueObject(PropertyJson)));
+	}
+	JsonObject->SetNumberField(TEXT("Rev"), RevisionCount);
+	JsonObject->SetArrayField(TEXT("Attrs"), JsonProperties);
+
+	auto JsonWriter = TJsonWriterFactory<TCHAR, TCondensedJsonPrintPolicy<TCHAR> >::Create(&JsonString);
+	FJsonSerializer::Serialize(JsonObject, JsonWriter);
+	JsonWriter->Close();
+}
+
+void FOnlinePartyData::FromJson(const FString& JsonString)
+{
+	// json string to key/val attrs
+	TSharedPtr<FJsonObject> JsonObject;
+	TSharedRef<TJsonReader<> > JsonReader = TJsonReaderFactory<>::Create(JsonString);
+	if (FJsonSerializer::Deserialize(JsonReader, JsonObject) &&
+		JsonObject.IsValid())
+	{
+		if (JsonObject->HasTypedField<EJson::Array>(TEXT("Attrs")))
+		{
+			const TArray<TSharedPtr<FJsonValue> >& JsonProperties = JsonObject->GetArrayField(TEXT("Attrs"));
+			for (auto JsonPropertyValue : JsonProperties)
+			{
+				TSharedPtr<FJsonObject> JsonPropertyObject = JsonPropertyValue->AsObject();
+				if (JsonPropertyObject.IsValid())
+				{
+					FString PropertyName;
+					if (JsonPropertyObject->TryGetStringField(TEXT("Name"), PropertyName) &&
+						!PropertyName.IsEmpty())
+					{
+						FVariantData PropertyData;
+						if (PropertyData.FromJson(JsonPropertyObject.ToSharedRef()))
+						{
+							KeyValAttrs.Add(PropertyName, PropertyData);
+						}
+					}
+				}
+			}
+		}
+
+		if (JsonObject->HasTypedField<EJson::Number>(TEXT("Rev")))
+		{
+			int32 NewRevisionCount = JsonObject->GetIntegerField(TEXT("Rev"));
+			if ((RevisionCount != 0) && (NewRevisionCount != RevisionCount) && (NewRevisionCount != (RevisionCount + 1)))
+			{
+				UE_LOG_ONLINEPARTY(Warning, TEXT("Unexpected revision received.  Current %d, new %d"), RevisionCount, NewRevisionCount);
+			}
+			RevisionCount = NewRevisionCount;
+		}
+	}
+}
diff --git a/Engine/Plugins/Online/OnlineSubsystem/Source/Private/OnlineSessionSettings.cpp b/Engine/Plugins/Online/OnlineSubsystem/Source/Private/OnlineSessionSettings.cpp
new file mode 100644
--- /dev/null
+++ b/Engine/Plugins/Online/OnlineSubsystem/Source/Private/OnlineSessionSettings.cpp
@@ -0,0 +1,208 @@
+// Copyright 1998-2016 Epic Games, Inc. All Rights Reserved.
+
+#include "OnlineSubsystemPrivatePCH.h"
+#include "OnlineSessionSettings.h"
+
+void DumpNamedSession(const FNamedOnlineSession* NamedSession)
+{
+	if (NamedSession != NULL)
+	{
+		UE_LOG(LogOnline, Verbose, TEXT("dumping NamedSession: "));
+		UE_LOG(LogOnline, Verbose, TEXT("	SessionName: %s"), *NamedSession->SessionName.ToString());	
+		UE_LOG(LogOnline, Verbose, TEXT("	HostingPlayerNum: %d"), NamedSession->HostingPlayerNum);
+		UE_LOG(LogOnline, Verbose, TEXT("	SessionState: %s"), EOnlineSessionState::ToString(NamedSession->SessionState));
+		UE_LOG(LogOnline, Verbose, TEXT("	RegisteredPlayers: "));
+		if (NamedSession->RegisteredPlayers.Num())
+		{
+			for (int32 UserIdx=0; UserIdx < NamedSession->RegisteredPlayers.Num(); UserIdx++)
+			{
+				UE_LOG(LogOnline, Verbose, TEXT("	    %d: %s"), UserIdx, *NamedSession->RegisteredPlayers[UserIdx]->ToDebugString());
+			}
+		}
+		else
+		{
+			UE_LOG(LogOnline, Verbose, TEXT("	    0 registered players"));
+		}
+
+		DumpSession(NamedSession);
+	}
+}
+
+void DumpSession(const FOnlineSession* Session)
+{
+	if (Session != NULL)
+	{
+		UE_LOG(LogOnline, Verbose, TEXT("dumping Session: "));
+		UE_LOG(LogOnline, Verbose, TEXT("	OwningPlayerName: %s"), *Session->OwningUserName);	
+		UE_LOG(LogOnline, Verbose, TEXT("	OwningPlayerId: %s"), Session->OwningUserId.IsValid() ? *Session->OwningUserId->ToDebugString() : TEXT("") );
+		UE_LOG(LogOnline, Verbose, TEXT("	NumOpenPrivateConnections: %d"), Session->NumOpenPrivateConnections);	
+		UE_LOG(LogOnline, Verbose, TEXT("	NumOpenPublicConnections: %d"), Session->NumOpenPublicConnections);	
+		UE_LOG(LogOnline, Verbose, TEXT("	SessionInfo: %s"), Session->SessionInfo.IsValid() ? *Session->SessionInfo->ToDebugString() : TEXT("NULL"));
+		DumpSessionSettings(&Session->SessionSettings);
+	}
+}
+
+void DumpSessionSettings(const FOnlineSessionSettings* SessionSettings)
+{
+	if (SessionSettings != NULL)
+	{
+		UE_LOG(LogOnline, Verbose, TEXT("dumping SessionSettings: "));
+		UE_LOG(LogOnline, Verbose, TEXT("\tNumPublicConnections: %d"), SessionSettings->NumPublicConnections);
+		UE_LOG(LogOnline, Verbose, TEXT("\tNumPrivateConnections: %d"), SessionSettings->NumPrivateConnections);
+		UE_LOG(LogOnline, Verbose, TEXT("\tbIsLanMatch: %s"), SessionSettings->bIsLANMatch ? TEXT("true") : TEXT("false"));
+		UE_LOG(LogOnline, Verbose, TEXT("\tbIsDedicated: %s"), SessionSettings->bIsDedicated ? TEXT("true") : TEXT("false"));
+		UE_LOG(LogOnline, Verbose, TEXT("\tbUsesStats: %s"), SessionSettings->bUsesStats ? TEXT("true") : TEXT("false"));
+		UE_LOG(LogOnline, Verbose, TEXT("\tbShouldAdvertise: %s"), SessionSettings->bShouldAdvertise ? TEXT("true") : TEXT("false"));
+		UE_LOG(LogOnline, Verbose, TEXT("\tbAllowJoinInProgress: %s"), SessionSettings->bAllowJoinInProgress ? TEXT("true") : TEXT("false"));
+		UE_LOG(LogOnline, Verbose, TEXT("\tbAllowInvites: %s"), SessionSettings->bAllowInvites ? TEXT("true") : TEXT("false"));
+		UE_LOG(LogOnline, Verbose, TEXT("\tbUsesPresence: %s"), SessionSettings->bUsesPresence ? TEXT("true") : TEXT("false"));
+		UE_LOG(LogOnline, Verbose, TEXT("\tbAllowJoinViaPresence: %s"), SessionSettings->bAllowJoinViaPresence ? TEXT("true") : TEXT("false"));
+		UE_LOG(LogOnline, Verbose, TEXT("\tbAllowJoinViaPresenceFriendsOnly: %s"), SessionSettings->bAllowJoinViaPresenceFriendsOnly ? TEXT("true") : TEXT("false"));
+		UE_LOG(LogOnline, Verbose, TEXT("\tBuildUniqueId: 0x%08x"), SessionSettings->BuildUniqueId);
+		UE_LOG(LogOnline, Verbose, TEXT("\tSettings:"));
+		for (FSessionSettings::TConstIterator It(SessionSettings->Settings); It; ++It)
+		{
+			FName Key = It.Key();
+			const FOnlineSessionSetting& Setting = It.Value();
+			UE_LOG(LogOnline, Verbose, TEXT("\t\t%s=%s"), *Key.ToString(), *Setting.ToString());
+		}
+	}
+}
+
+template<typename ValueType> 
+void FOnlineSessionSettings::Set(FName Key, const ValueType& Value, EOnlineDataAdvertisementType::Type InType)
+{
+	FOnlineSessionSetting* Setting = Settings.Find(Key);
+	if (Setting)
+	{
+		Setting->Data.SetValue(Value);
+		Setting->AdvertisementType = InType;
+	}
+	else
+	{
+		Settings.Add(Key, FOnlineSessionSetting(Value, InType));
+	}
+}
+
+/** Explicit instantiation of supported types to Set template above */
+#if !UE_BUILD_DOCS
+template ONLINESUBSYSTEM_API void FOnlineSessionSettings::Set(FName Key, const int32& Value, EOnlineDataAdvertisementType::Type InType);
+template ONLINESUBSYSTEM_API void FOnlineSessionSettings::Set(FName Key, const float& Value, EOnlineDataAdvertisementType::Type InType);
+template ONLINESUBSYSTEM_API void FOnlineSessionSettings::Set(FName Key, const uint64& Value, EOnlineDataAdvertisementType::Type InType);
+template ONLINESUBSYSTEM_API void FOnlineSessionSettings::Set(FName Key, const double& Value, EOnlineDataAdvertisementType::Type InType);
+template ONLINESUBSYSTEM_API void FOnlineSessionSettings::Set(FName Key, const FString& Value, EOnlineDataAdvertisementType::Type InType);
+template ONLINESUBSYSTEM_API void FOnlineSessionSettings::Set(FName Key, const bool& Value, EOnlineDataAdvertisementType::Type InType);
+template ONLINESUBSYSTEM_API void FOnlineSessionSettings::Set(FName Key, const TArray<uint8>& Value, EOnlineDataAdvertisementType::Type InType);
+#endif
+
+void FOnlineSessionSettings::Set(FName Key, const FOnlineSessionSetting& SrcSetting)
+{
+	FOnlineSessionSetting* Setting = Settings.Find(Key);
+	if (Setting)
+	{
+		Setting->Data = SrcSetting.Data;
+		Setting->AdvertisementType = SrcSetting.AdvertisementType;
+	}
+	else
+	{
+		Settings.Add(Key, SrcSetting);
+	}
+}
+
+template<typename ValueType> 
+bool FOnlineSessionSettings::Get(FName Key, ValueType& Value) const
+{
+	const FOnlineSessionSetting* Setting = Settings.Find(Key);
+	if (Setting)
+	{
+		Setting->Data.GetValue(Value);
+		return true;
+	}
+
+	return false;
+}
+
+/** Explicit instantiation of supported types to Get template above */
+template ONLINESUBSYSTEM_API bool FOnlineSessionSettings::Get(FName Key, int32& Value) const;
+template ONLINESUBSYSTEM_API bool FOnlineSessionSettings::Get(FName Key, float& Value) const;
+template ONLINESUBSYSTEM_API bool FOnlineSessionSettings::Get(FName Key, uint64& Value) const;
+template ONLINESUBSYSTEM_API bool FOnlineSessionSettings::Get(FName Key, double& Value) const;
+template ONLINESUBSYSTEM_API bool FOnlineSessionSettings::Get(FName Key, bool& Value) const;
+template ONLINESUBSYSTEM_API bool FOnlineSessionSettings::Get(FName Key, FString& Value) const;
+template ONLINESUBSYSTEM_API bool FOnlineSessionSettings::Get(FName Key, TArray<uint8>& Value) const;
+
+bool FOnlineSessionSettings::Remove(FName Key)
+{
+	return Settings.Remove(Key) > 0;
+}
+
+EOnlineDataAdvertisementType::Type FOnlineSessionSettings::GetAdvertisementType(FName Key) const
+{
+	const FOnlineSessionSetting* Setting = Settings.Find(Key);
+	if (Setting)
+	{
+		return Setting->AdvertisementType;
+	}
+
+	UE_LOG(LogOnline, Warning, TEXT("Unable to find key for advertisement type request: %s"), *Key.ToString());
+	return EOnlineDataAdvertisementType::DontAdvertise;
+}
+
+template<typename ValueType> 
+void FOnlineSearchSettings::Set(FName Key, const ValueType& Value, EOnlineComparisonOp::Type InType)
+{
+	FOnlineSessionSearchParam* SearchParam = SearchParams.Find(Key);
+	if (SearchParam)
+	{
+		SearchParam->Data.SetValue(Value);
+		SearchParam->ComparisonOp = InType;
+	}
+	else
+	{
+		SearchParams.Add(Key, FOnlineSessionSearchParam(Value, InType));
+	}
+}
+
+/** Explicit instantiation of supported types to Set template above */
+template ONLINESUBSYSTEM_API void FOnlineSearchSettings::Set<int32>(FName Key, const int32& Value, EOnlineComparisonOp::Type InType);
+template ONLINESUBSYSTEM_API void FOnlineSearchSettings::Set<float>(FName Key, const float& Value, EOnlineComparisonOp::Type InType);
+template ONLINESUBSYSTEM_API void FOnlineSearchSettings::Set<uint64>(FName Key, const uint64& Value, EOnlineComparisonOp::Type InType);
+template ONLINESUBSYSTEM_API void FOnlineSearchSettings::Set<double>(FName Key, const double& Value, EOnlineComparisonOp::Type InType);
+template ONLINESUBSYSTEM_API void FOnlineSearchSettings::Set<FString>(FName Key, const FString& Value, EOnlineComparisonOp::Type InType);
+template ONLINESUBSYSTEM_API void FOnlineSearchSettings::Set< TArray<uint8> >(FName Key, const TArray<uint8>& Value, EOnlineComparisonOp::Type InType);
+template ONLINESUBSYSTEM_API void FOnlineSearchSettings::Set<bool>(FName Key, const bool& Value, EOnlineComparisonOp::Type InType);
+
+template<typename ValueType> 
+bool FOnlineSearchSettings::Get(FName Key, ValueType& Value) const
+{
+	const FOnlineSessionSearchParam* SearchParam = SearchParams.Find(Key);
+	if (SearchParam)
+	{
+		SearchParam->Data.GetValue(Value);
+		return true;
+	}
+
+	return false;
+}
+
+/** Explicit instantiation of supported types to Get template above */
+template ONLINESUBSYSTEM_API bool FOnlineSearchSettings::Get<int32>(FName Key, int32& Value) const;
+template ONLINESUBSYSTEM_API bool FOnlineSearchSettings::Get<float>(FName Key, float& Value) const;
+template ONLINESUBSYSTEM_API bool FOnlineSearchSettings::Get<uint64>(FName Key, uint64& Value) const;
+template ONLINESUBSYSTEM_API bool FOnlineSearchSettings::Get<double>(FName Key, double& Value) const;
+template ONLINESUBSYSTEM_API bool FOnlineSearchSettings::Get<FString>(FName Key, FString& Value) const;
+template ONLINESUBSYSTEM_API bool FOnlineSearchSettings::Get< TArray<uint8> >(FName Key, TArray<uint8>& Value) const;
+template ONLINESUBSYSTEM_API bool FOnlineSearchSettings::Get<bool>(FName Key, bool& Value) const;
+
+EOnlineComparisonOp::Type FOnlineSearchSettings::GetComparisonOp(FName Key) const
+{
+	const FOnlineSessionSearchParam* SearchParam = SearchParams.Find(Key);
+	if (SearchParam)
+	{
+		return SearchParam->ComparisonOp;
+	}
+
+	UE_LOG(LogOnline, Warning, TEXT("Unable to find key for comparison op request: %s"), *Key.ToString());
+	return EOnlineComparisonOp::Equals;
+}
+
diff --git a/Engine/Plugins/Online/OnlineSubsystem/Source/Private/OnlineSubsystem.cpp b/Engine/Plugins/Online/OnlineSubsystem/Source/Private/OnlineSubsystem.cpp
new file mode 100644
--- /dev/null
+++ b/Engine/Plugins/Online/OnlineSubsystem/Source/Private/OnlineSubsystem.cpp
@@ -0,0 +1,133 @@
+// Copyright 1998-2016 Epic Games, Inc. All Rights Reserved.
+
+#include "OnlineSubsystemPrivatePCH.h"
+#include "OnlineSessionInterface.h"
+#include "OnlineIdentityInterface.h"
+#include "OnlineAchievementsInterface.h"
+#include "NboSerializer.h"
+#include "Online.h"
+#include "IConsoleManager.h"
+#include "NetworkVersion.h"
+
+DEFINE_LOG_CATEGORY(LogOnline);
+DEFINE_LOG_CATEGORY(LogOnlineGame);
+DEFINE_LOG_CATEGORY(LogOnlineChat);
+
+#if STATS
+ONLINESUBSYSTEM_API DEFINE_STAT(STAT_Online_Async);
+ONLINESUBSYSTEM_API DEFINE_STAT(STAT_Online_AsyncTasks);
+ONLINESUBSYSTEM_API DEFINE_STAT(STAT_Session_Interface);
+ONLINESUBSYSTEM_API DEFINE_STAT(STAT_Voice_Interface);
+#endif
+
+int32 GetBuildUniqueId()
+{
+	static bool bStaticCheck = false;
+	static bool bUseBuildIdOverride = false;
+	static int32 BuildIdOverride = 0;
+	if (!bStaticCheck)
+	{
+		if (FParse::Value(FCommandLine::Get(), TEXT("BuildIdOverride="), BuildIdOverride) && BuildIdOverride != 0)
+		{
+			bUseBuildIdOverride = true;
+		}
+		else
+		{
+			if (!GConfig->GetBool(TEXT("OnlineSubsystem"), TEXT("bUseBuildIdOverride"), bUseBuildIdOverride, GEngineIni))
+			{
+				UE_LOG_ONLINE(Warning, TEXT("Missing bUseBuildIdOverride= in [OnlineSubsystem] of DefaultEngine.ini"));
+			}
+
+			if (!GConfig->GetInt(TEXT("OnlineSubsystem"), TEXT("BuildIdOverride"), BuildIdOverride, GEngineIni))
+			{
+				UE_LOG_ONLINE(Warning, TEXT("Missing BuildIdOverride= in [OnlineSubsystem] of DefaultEngine.ini"));
+			}
+		}
+
+		bStaticCheck = true;
+	}
+
+	const uint32 NetworkVersion = FNetworkVersion::GetLocalNetworkVersion();
+
+	int32 BuildId = 0;
+	if (bUseBuildIdOverride == false)
+	{
+		/** Engine package CRC doesn't change, can't be used as the version - BZ */
+		FNboSerializeToBuffer Buffer(64);
+		// Serialize to a NBO buffer for consistent CRCs across platforms
+		Buffer << NetworkVersion;
+		// Now calculate the CRC
+		uint32 Crc = FCrc::MemCrc32((uint8*)Buffer, Buffer.GetByteCount());
+
+		// make sure it's positive when it's cast back to an int
+		BuildId = static_cast<int32>(Crc & 0x7fffffff);
+	}
+	else
+	{
+		BuildId = BuildIdOverride;
+	}
+
+	UE_LOG_ONLINE(VeryVerbose, TEXT("GetBuildUniqueId: Network CL %u LocalNetworkVersion %u bUseBuildIdOverride %d BuildIdOverride %d BuildId %d"),
+		FNetworkVersion::GetNetworkCompatibleChangelist(),
+		NetworkVersion,
+		bUseBuildIdOverride,
+		BuildIdOverride,
+		BuildId);
+
+	return BuildId;
+}
+
+bool IsPlayerInSessionImpl(IOnlineSession* SessionInt, FName SessionName, const FUniqueNetId& UniqueId)
+{
+	bool bFound = false;
+	FNamedOnlineSession* Session = SessionInt->GetNamedSession(SessionName);
+	if (Session != NULL)
+	{
+		const bool bIsSessionOwner = *Session->OwningUserId == UniqueId;
+
+		FUniqueNetIdMatcher PlayerMatch(UniqueId);
+		if (bIsSessionOwner || 
+			Session->RegisteredPlayers.IndexOfByPredicate(PlayerMatch) != INDEX_NONE)
+		{
+			bFound = true;
+		}
+	}
+	return bFound;
+}
+
+
+#if !UE_BUILD_SHIPPING
+
+static void ResetAchievements()
+{
+	auto IdentityInterface = Online::GetIdentityInterface();
+	if (!IdentityInterface.IsValid())
+	{
+		UE_LOG_ONLINE(Warning, TEXT("ResetAchievements command: couldn't get the identity interface"));
+		return;
+	}
+	
+	TSharedPtr<const FUniqueNetId> UserId = IdentityInterface->GetUniquePlayerId(0);
+	if(!UserId.IsValid())
+	{
+		UE_LOG_ONLINE(Warning, TEXT("ResetAchievements command: invalid UserId"));
+		return;
+	}
+
+	auto AchievementsInterface = Online::GetAchievementsInterface();
+	if (!AchievementsInterface.IsValid())
+	{
+		UE_LOG_ONLINE(Warning, TEXT("ResetAchievements command: couldn't get the achievements interface"));
+		return;
+	}
+
+	AchievementsInterface->ResetAchievements(*UserId);
+}
+
+FAutoConsoleCommand CmdResetAchievements(
+	TEXT("online.ResetAchievements"),
+	TEXT("Reset achievements for the currently logged in user."),
+	FConsoleCommandDelegate::CreateStatic(ResetAchievements)
+	);
+
+#endif
diff --git a/Engine/Plugins/Online/OnlineSubsystem/Source/Private/OnlineSubsystemModule.cpp b/Engine/Plugins/Online/OnlineSubsystem/Source/Private/OnlineSubsystemModule.cpp
new file mode 100644
--- /dev/null
+++ b/Engine/Plugins/Online/OnlineSubsystem/Source/Private/OnlineSubsystemModule.cpp
@@ -0,0 +1,297 @@
+// Copyright 1998-2016 Epic Games, Inc. All Rights Reserved.
+
+#include "OnlineSubsystemPrivatePCH.h"
+#include "ModuleManager.h"
+#include "OnlineSubsystemImpl.h"
+
+IMPLEMENT_MODULE( FOnlineSubsystemModule, OnlineSubsystem );
+
+/** Helper function to turn the friendly subsystem name into the module name */
+static inline FName GetOnlineModuleName(const FString& SubsystemName)
+{
+	FString ModuleBase(TEXT("OnlineSubsystem"));
+
+	FName ModuleName;
+	if (!SubsystemName.StartsWith(ModuleBase, ESearchCase::CaseSensitive))
+	{
+		ModuleName = FName(*(ModuleBase + SubsystemName));
+	}
+	else
+	{
+		ModuleName = FName(*SubsystemName);
+	}
+
+	return ModuleName;
+}
+
+/**
+ * Helper function that loads a given platform service module if it isn't already loaded
+ *
+ * @param SubsystemName Name of the requested platform service to load
+ * @return The module interface of the requested platform service, NULL if the service doesn't exist
+ */
+static TSharedPtr<IModuleInterface> LoadSubsystemModule(const FString& SubsystemName)
+{
+#if !UE_BUILD_SHIPPING && !UE_BUILD_SHIPPING_WITH_EDITOR
+	// Early out if we are overriding the module load
+	bool bAttemptLoadModule = !FParse::Param(FCommandLine::Get(), *FString::Printf(TEXT("no%s"), *SubsystemName));
+	if (bAttemptLoadModule)
+#endif
+	{
+		FName ModuleName;
+		FModuleManager& ModuleManager = FModuleManager::Get();
+
+		ModuleName = GetOnlineModuleName(SubsystemName);
+		if (!ModuleManager.IsModuleLoaded(ModuleName))
+		{
+			// Attempt to load the module
+			ModuleManager.LoadModule(ModuleName);
+		}
+
+		return ModuleManager.GetModule(ModuleName);
+	}
+
+	return NULL;
+}
+
+void FOnlineSubsystemModule::StartupModule()
+{
+	LoadDefaultSubsystem();
+	// Also load the console/platform specific OSS which might not necessarily be the default OSS instance
+	IOnlineSubsystem::GetByPlatform();
+}
+
+void FOnlineSubsystemModule::ShutdownModule()
+{
+	ShutdownOnlineSubsystem();
+}
+
+void FOnlineSubsystemModule::LoadDefaultSubsystem()
+{
+	FString InterfaceString;
+
+	// Load the platform defined "default" online services module
+	if (GConfig->GetString(TEXT("OnlineSubsystem"), TEXT("DefaultPlatformService"), InterfaceString, GEngineIni) &&
+		InterfaceString.Len() > 0)
+	{
+		FName InterfaceName = FName(*InterfaceString);
+		// A module loaded with its factory method set for creation and a default instance of the online subsystem is required
+		if (LoadSubsystemModule(InterfaceString).IsValid() &&
+			OnlineFactories.Contains(InterfaceName) &&
+			GetOnlineSubsystem(InterfaceName) != NULL)
+		{
+			DefaultPlatformService = InterfaceName;
+		}
+		else
+		{
+			UE_LOG(LogOnline, Log, TEXT("Unable to load default OnlineSubsystem module %s, using NULL interface"), *InterfaceString);
+			InterfaceString = TEXT("Null");
+			InterfaceName = FName(*InterfaceString);
+
+			// A module loaded with its factory method set for creation and a default instance of the online subsystem is required
+			if (LoadSubsystemModule(InterfaceString).IsValid() &&
+				OnlineFactories.Contains(InterfaceName) &&
+				GetOnlineSubsystem(InterfaceName) != NULL)
+			{
+				DefaultPlatformService = InterfaceName;
+			}
+		}
+	}
+	else
+	{
+		UE_LOG(LogOnline, Log, TEXT("No default platform service specified for OnlineSubsystem"));
+	}
+}
+
+void FOnlineSubsystemModule::ReloadDefaultSubsystem()
+{
+	DestroyOnlineSubsystem(DefaultPlatformService);
+	LoadDefaultSubsystem();
+}
+
+void FOnlineSubsystemModule::ShutdownOnlineSubsystem()
+{
+	FModuleManager& ModuleManager = FModuleManager::Get();
+
+	// Shutdown all online subsystem instances
+	for (TMap<FName, IOnlineSubsystemPtr>::TIterator It(OnlineSubsystems); It; ++It)
+	{
+		It.Value()->Shutdown();
+	}
+	OnlineSubsystems.Empty();
+
+	// Unload all the supporting factories
+	for (TMap<FName, IOnlineFactory*>::TIterator It(OnlineFactories); It; ++It)
+	{
+		UE_LOG(LogOnline, Display, TEXT("Unloading online subsystem: %s"), *It.Key().ToString());
+
+		// Unloading the module will do proper cleanup
+		FName ModuleName = GetOnlineModuleName(It.Key().ToString());
+
+		const bool bIsShutdown = true;
+		ModuleManager.UnloadModule(ModuleName, bIsShutdown);
+	} 
+	//ensure(OnlineFactories.Num() == 0);
+}
+
+void FOnlineSubsystemModule::RegisterPlatformService(const FName FactoryName, IOnlineFactory* Factory)
+{
+	if (!OnlineFactories.Contains(FactoryName))
+	{
+		OnlineFactories.Add(FactoryName, Factory);
+	}
+}
+
+void FOnlineSubsystemModule::UnregisterPlatformService(const FName FactoryName)
+{
+	if (OnlineFactories.Contains(FactoryName))
+	{
+		OnlineFactories.Remove(FactoryName);
+	}
+}
+
+FName FOnlineSubsystemModule::ParseOnlineSubsystemName(const FName& FullName, FName& SubsystemName, FName& InstanceName) const
+{
+#if !(UE_GAME || UE_SERVER)
+	SubsystemName = DefaultPlatformService;
+	InstanceName = FOnlineSubsystemImpl::DefaultInstanceName;
+
+	if (FullName != NAME_None)
+	{
+		FString FullNameStr = FullName.ToString();
+
+		int32 DelimIdx = INDEX_NONE;
+		static const TCHAR InstanceDelim = ':';
+		if (FullNameStr.FindChar(InstanceDelim, DelimIdx))
+		{
+			if (DelimIdx > 0)
+			{
+				SubsystemName = FName(*FullNameStr.Left(DelimIdx));
+			}
+
+			if ((DelimIdx + 1) < FullNameStr.Len())
+			{
+				InstanceName = FName(*FullNameStr.RightChop(DelimIdx + 1));
+			}
+		}
+		else
+		{
+			SubsystemName = FName(*FullNameStr);
+		}
+	}
+
+	return FName(*FString::Printf(TEXT("%s:%s"), *SubsystemName.ToString(), *InstanceName.ToString()));
+#else	
+	
+	SubsystemName = FullName == NAME_None ? DefaultPlatformService : FullName;
+	InstanceName = FOnlineSubsystemImpl::DefaultInstanceName;
+
+#if !UE_BUILD_SHIPPING
+	int32 DelimIdx = INDEX_NONE;
+	static const TCHAR InstanceDelim = ':';
+	ensure(!FullName.ToString().FindChar(InstanceDelim, DelimIdx) && DelimIdx == INDEX_NONE);
+#endif
+	return SubsystemName;
+#endif // !(UE_GAME || UE_SERVER)
+}
+
+IOnlineSubsystem* FOnlineSubsystemModule::GetOnlineSubsystem(const FName InSubsystemName)
+{
+	FName SubsystemName, InstanceName;
+	FName KeyName = ParseOnlineSubsystemName(InSubsystemName, SubsystemName, InstanceName);
+
+	IOnlineSubsystemPtr* OnlineSubsystem = NULL;
+	if (SubsystemName != NAME_None)
+	{
+		OnlineSubsystem = OnlineSubsystems.Find(KeyName);
+		if (OnlineSubsystem == NULL)
+		{
+			IOnlineFactory** OSSFactory = OnlineFactories.Find(SubsystemName);
+			if (OSSFactory == NULL)
+			{
+				// Attempt to load the requested factory
+				TSharedPtr<IModuleInterface> NewModule = LoadSubsystemModule(SubsystemName.ToString());
+				if (NewModule.IsValid())
+				{
+					// If the module loaded successfully this should be non-NULL
+					OSSFactory = OnlineFactories.Find(SubsystemName);
+				}
+			}
+
+			if (OSSFactory != NULL)
+			{
+				IOnlineSubsystemPtr NewSubsystemInstance = (*OSSFactory)->CreateSubsystem(InstanceName);
+				if (NewSubsystemInstance.IsValid())
+				{
+					OnlineSubsystems.Add(KeyName, NewSubsystemInstance);
+					OnlineSubsystem = OnlineSubsystems.Find(KeyName);
+				}
+				else
+				{
+						bool* bNotedPreviously = OnlineSubsystemFailureNotes.Find(KeyName);
+						if (!bNotedPreviously || !(*bNotedPreviously))
+						{
+							UE_LOG(LogOnline, Log, TEXT("Unable to create OnlineSubsystem module %s"), *SubsystemName.ToString());
+							OnlineSubsystemFailureNotes.Add(KeyName, true);
+						}
+					}
+				}
+			}
+	}
+
+	return (OnlineSubsystem == NULL) ? NULL : (*OnlineSubsystem).Get();
+}
+
+void FOnlineSubsystemModule::DestroyOnlineSubsystem(const FName InSubsystemName)
+{
+	FName SubsystemName, InstanceName;
+	FName KeyName = ParseOnlineSubsystemName(InSubsystemName, SubsystemName, InstanceName);
+
+	if (SubsystemName != NAME_None)
+	{
+		IOnlineSubsystemPtr OnlineSubsystem;
+		OnlineSubsystems.RemoveAndCopyValue(KeyName, OnlineSubsystem);
+		if (OnlineSubsystem.IsValid())
+		{
+			OnlineSubsystem->Shutdown();
+			OnlineSubsystemFailureNotes.Remove(KeyName);
+		}
+		else
+		{
+			UE_LOG(LogOnline, Warning, TEXT("OnlineSubsystem instance %s not found, unable to destroy."), *KeyName.ToString());
+		}
+	}
+}
+
+bool FOnlineSubsystemModule::DoesInstanceExist(const FName InSubsystemName) const
+{
+	bool bIsLoaded = false;
+
+	FName SubsystemName, InstanceName;
+	FName KeyName = ParseOnlineSubsystemName(InSubsystemName, SubsystemName, InstanceName);
+	if (SubsystemName != NAME_None)
+	{
+		const IOnlineSubsystemPtr* OnlineSubsystem = OnlineSubsystems.Find(KeyName);
+		return OnlineSubsystem && OnlineSubsystem->IsValid() ? true : false;
+	}
+
+	return false;
+}
+
+bool FOnlineSubsystemModule::IsOnlineSubsystemLoaded(const FName InSubsystemName) const
+{
+	bool bIsLoaded = false;
+
+	FName SubsystemName, InstanceName;
+	ParseOnlineSubsystemName(InSubsystemName, SubsystemName, InstanceName);
+
+	if (SubsystemName != NAME_None)
+	{
+		if (FModuleManager::Get().IsModuleLoaded(GetOnlineModuleName(SubsystemName.ToString())))
+		{
+			bIsLoaded = true;
+		}
+	}
+	return bIsLoaded;
+}
+
diff --git a/Engine/Plugins/Online/OnlineSubsystem/Source/Private/OnlineSubsystemPrivatePCH.h b/Engine/Plugins/Online/OnlineSubsystem/Source/Private/OnlineSubsystemPrivatePCH.h
new file mode 100644
--- /dev/null
+++ b/Engine/Plugins/Online/OnlineSubsystem/Source/Private/OnlineSubsystemPrivatePCH.h
@@ -0,0 +1,12 @@
+// Copyright 1998-2016 Epic Games, Inc. All Rights Reserved.
+
+#pragma once
+
+#include "Core.h"
+#include "CoreUObject.h"
+#include "OnlineSubsystem.h"
+#include "OnlineSubsystemModule.h"
+#include "SocketSubsystem.h"
+#include "ModuleManager.h"
+
+
diff --git a/Engine/Plugins/Online/OnlineSubsystem/Source/Private/TurnBasedMatchInterface.cpp b/Engine/Plugins/Online/OnlineSubsystem/Source/Private/TurnBasedMatchInterface.cpp
new file mode 100644
--- /dev/null
+++ b/Engine/Plugins/Online/OnlineSubsystem/Source/Private/TurnBasedMatchInterface.cpp
@@ -0,0 +1,9 @@
+// Copyright 1998-2016 Epic Games, Inc. All Rights Reserved.
+
+#include "OnlineSubsystemPrivatePCH.h"
+#include "TurnBasedMatchInterface.h"
+
+UTurnBasedMatchInterface::UTurnBasedMatchInterface(const FObjectInitializer& ObjectInitializer)
+: Super(ObjectInitializer)
+{
+}
\ No newline at end of file
diff --git a/Engine/Plugins/Online/OnlineSubsystem/Source/Public/Interfaces/IMessageSanitizerInterface.h b/Engine/Plugins/Online/OnlineSubsystem/Source/Public/Interfaces/IMessageSanitizerInterface.h
new file mode 100644
--- /dev/null
+++ b/Engine/Plugins/Online/OnlineSubsystem/Source/Public/Interfaces/IMessageSanitizerInterface.h
@@ -0,0 +1,41 @@
+// Copyright 1998-2015 Epic Games, Inc. All Rights Reserved.
+
+#pragma once
+
+DECLARE_DELEGATE_TwoParams(FOnMessageProcessed, bool /*Process success*/, const FString& /*SanitizedMessage*/);
+DECLARE_DELEGATE_TwoParams(FOnMessageArrayProcessed, bool /*Process success*/, const TArray<FString>& /*SanitizedMessages*/);
+
+struct FSanitizeMessage
+{
+	FSanitizeMessage(const FString& InRawMessage, FOnMessageProcessed InProcessCompleteDelegate)
+		: RawMessage(InRawMessage)
+		, CompletionDelegate(InProcessCompleteDelegate)
+	{}
+
+	FString RawMessage;
+	FOnMessageProcessed CompletionDelegate;
+};
+
+struct FMultiPartMessage
+{
+	TArray<FString> AlreadyProcessedMessages;
+	TArray<int32> AlreadyProcessedIndex;
+	TArray<FString> SanitizedMessages;
+	FString MessageToSanatize;
+	bool bCompleted;
+	int32 PartNumber;
+};
+
+class IMessageSanitizer
+{
+protected:
+	IMessageSanitizer() {};
+
+public:
+	virtual ~IMessageSanitizer() {};
+	virtual void SanitizeMessage(const FString& Message, FOnMessageProcessed CompletionDelegate) = 0;
+	virtual void SanitizeMessageArray(const TArray<FString>& MessageArray, FOnMessageArrayProcessed CompletionDelegate) = 0;
+};
+
+typedef TSharedPtr<IMessageSanitizer, ESPMode::ThreadSafe> IMessageSanitizerPtr;
+
