Copying //Tasks/UE4/Dev-Online-Plugins-v3 to //UE4/Dev-Main (Source: //Tasks/UE4/Dev-Online-Plugins-v3 @ 3047294)

diff --git a/Engine/Plugins/Online/OnlineSubsystem/Source/Public/Interfaces/OnlinePartyInterface.h b/Engine/Plugins/Online/OnlineSubsystem/Source/Public/Interfaces/OnlinePartyInterface.h
new file mode 100644
--- /dev/null
+++ b/Engine/Plugins/Online/OnlineSubsystem/Source/Public/Interfaces/OnlinePartyInterface.h
@@ -0,0 +1,1918 @@
+// Copyright 1998-2016 Epic Games, Inc. All Rights Reserved.
+
+#pragma once
+
+#include "OnlineSubsystemTypes.h"
+#include "OnlineChatInterface.h"
+
+#define F_PREFIX(TypeToPrefix) F##TypeToPrefix
+#define PARTY_DECLARE_DELEGATETYPE(Type) typedef F##Type::FDelegate F##Type##Delegate
+
+ONLINESUBSYSTEM_API DECLARE_LOG_CATEGORY_EXTERN(LogOnlineParty, Display, All);
+
+#define ONLINEPARTY_LOG_PREFIX TEXT("OSS::Party: ")
+#define UE_LOG_ONLINEPARTY(Verbosity, Format, ...) \
+{ \
+	UE_LOG(LogOnlineParty, Verbosity, TEXT("%s%s"), ONLINEPARTY_LOG_PREFIX, *FString::Printf(Format, ##__VA_ARGS__)); \
+}
+
+
+/**
+ * Party member user info returned by IOnlineParty interface
+ */
+class FOnlinePartyMember : public FOnlineUser
+{
+public:
+};
+
+/**
+ * Data associated with the entire party
+ */
+class ONLINESUBSYSTEM_API FOnlinePartyData : public TSharedFromThis<FOnlinePartyData>
+{
+public:
+	FOnlinePartyData() :
+		TotalBytes(0),
+		TotalEffectiveBytes(0),
+		TotalPackets(0),
+		RevisionCount(0)
+		{}
+	~FOnlinePartyData() {}
+
+	/**
+	 * Get an attribute from the party data
+	 *
+	 * @param AttrName - key for the attribute
+	 * @param OutAttrValue - [out] value for the attribute if found
+	 *
+	 * @return true if the attribute was found
+	 */
+	bool GetAttribute(const FString& AttrName, FVariantData& OutAttrValue) const
+	{
+		bool bResult = false;
+
+		const FVariantData* FoundValuePtr = KeyValAttrs.Find(AttrName);
+		if (FoundValuePtr != nullptr)
+		{
+			OutAttrValue = *FoundValuePtr;
+		}
+
+		return bResult;
+	}
+
+	/**
+	 * Set an attribute from the party data
+	 *
+	 * @param AttrName - key for the attribute
+	 * @param AttrValue - value to set the attribute to
+	 */
+	void SetAttribute(const FString& AttrName, const FVariantData& AttrValue)
+	{
+		FVariantData& NewAttrValue = KeyValAttrs.FindOrAdd(AttrName);
+		if (NewAttrValue != AttrValue)
+		{
+			NewAttrValue = AttrValue;
+			DirtyKeys.Add(AttrName);
+		}
+	}
+
+	/**
+	 * Clear the attributes map
+	 */
+	void ClearAttributes()
+	{
+		KeyValAttrs.Empty();
+		DirtyKeys.Empty();
+	}
+
+	/** 
+	 * Clear the dirty keys set, called after successfully sending an update of the dirty elements
+	 */
+	void ClearDirty()
+	{
+		DirtyKeys.Empty();
+	}
+
+	/** 
+	 * Increment the stat tracking variables on packet sent
+	 * 
+	 * @param PacketSize - size of the packet generated
+	 * @param NumRecipients - number of recipients the packet was sent to
+	 * @param bIncrementRevision - this packet was a dirty packet update so we should increment the revision
+	 */
+	void OnPacketSent(int32 PacketSize, int32 NumRecipients, bool bIncrementRevision) const
+	{
+		TotalPackets++;
+		TotalBytes += PacketSize;
+		TotalEffectiveBytes += PacketSize * NumRecipients;
+		if (bIncrementRevision)
+		{
+			++RevisionCount;
+		}
+	}
+
+	/** 
+	 * Generate a JSON packet containing all key-value attributes
+	 * 
+	 * @param JsonString - [out] string containing the resulting JSON output
+	 * 
+	 */
+	void ToJsonFull(FString& JsonString) const;
+	
+	/** 
+	 * Generate a JSON packet containing only the dirty key-value attributes for a delta update
+	 *
+	 * @param JsonString - [out] string containing the resulting JSON output
+	 *
+	 */
+	void ToJsonDirty(FString& JsonString) const;
+
+	/** 
+	 * Update attributes from a JSON packet
+	 *
+	 * @param JsonString - string containing the JSON packet
+	 *
+	 */
+	void FromJson(const FString& JsonString);
+
+	/** Accessor functions for KeyValAttrs map */
+	FOnlineKeyValuePairs<FString, FVariantData>& GetKeyValAttrs() { return KeyValAttrs; }
+	const FOnlineKeyValuePairs<FString, FVariantData>& GetKeyValAttrs() const { return KeyValAttrs; }
+
+	/** Stat tracking variables */
+	/** Total number of bytes generated by calls to ToJsonFull and ToJsonDirty */
+	mutable int32 TotalBytes;
+	/** Total number of bytes generated by calls to ToJsonFull and ToJsonDirty, multiplied by the number of recipients the packet was sent to */
+	mutable int32 TotalEffectiveBytes;
+	/** Total number of packets generated by calls to ToJsonFull and ToJsonDirty */
+	mutable int32 TotalPackets;
+
+	/** Id representing number of updates sent, useful for determining if a client has missed an update */
+	mutable int32 RevisionCount;
+
+private:
+	/** map of key/val attributes that represents the data */
+	FOnlineKeyValuePairs<FString, FVariantData>  KeyValAttrs;
+
+	/** set of which fields are dirty and need to transmitted */
+	TSet<FString> DirtyKeys;
+};
+
+/**
+* Info needed to join a party
+*/
+class IOnlinePartyPendingJoinRequestInfo
+{
+public:
+	IOnlinePartyPendingJoinRequestInfo() {}
+	virtual ~IOnlinePartyPendingJoinRequestInfo() {}
+
+	/**
+	* @return id of the sender of this join request
+	*/
+	virtual const TSharedRef<const FUniqueNetId>& GetSenderId() const = 0;
+
+	/**
+	* @return display name of the sender of this join request
+	*/
+	virtual const FString& GetSenderDisplayName() const = 0;
+};
+
+/**
+ * Info needed to join a party
+ */
+class IOnlinePartyJoinInfo
+{
+public:
+	IOnlinePartyJoinInfo() {}
+	virtual ~IOnlinePartyJoinInfo() {}
+
+	virtual bool IsValid() const = 0;
+
+	/**
+	 * @return party id of party associated with this join invite
+	 */
+	virtual const TSharedRef<const FOnlinePartyId>& GetPartyId() const = 0;
+
+	/**
+	 * @return party id of party associated with this join invite
+	 */
+	virtual const FOnlinePartyTypeId GetPartyTypeId() const = 0;
+
+	/**
+	 * @return user id of the leader of the party associated with this join info
+	 */
+	virtual const TSharedRef<const FUniqueNetId>& GetLeaderId() const = 0;
+
+	/**
+	 * @return user id of the leader of the party associated with this join info
+	 */
+	virtual const FString& GetLeaderDisplayName() const = 0;
+
+	/**
+	 * @return user id of where this join info came from
+	 */
+	virtual const TSharedRef<const FUniqueNetId>& GetSourceUserId() const = 0;
+
+	/**
+	 * @return user id of where this join info came from
+	 */
+	virtual const FString& GetSourceDisplayName() const = 0;
+
+	/**
+	 * @return true if the join info has some form of key(does not guarantee the validity of that key)
+	 */
+	virtual bool HasKey() const = 0;
+
+	/**
+	 * @return true if a password can be used to bypass generated access key
+     */
+	virtual bool HasPassword() const = 0;
+
+	/**
+	 * @return true if the party is known to be accepting members
+     */
+	virtual bool IsAcceptingMembers() const = 0;
+
+	/**
+	 * @return true if this is a party of one
+	 */
+	virtual bool IsPartyOfOne() const = 0;
+
+	/**
+	 * @return why the party is not accepting members
+	 */
+	virtual int32 GetNotAcceptingReason() const = 0;
+
+	/**
+	 * @return id of the client app associated with the sender of the party invite
+	 */
+	virtual const FString& GetAppId() const = 0;
+
+	/**
+	* @return id of the build associated with the sender of the party invite
+	*/
+	virtual const FString& GetBuildId() const = 0;
+
+	/**
+	 * @return client specify key-value data
+	 */
+	virtual const FOnlinePartyData& GetClientData() const = 0;
+
+	/**
+	 * @return whether or not the join info can be used to join
+	 */
+	virtual bool CanJoin() const = 0;
+
+	/**
+	 * @return whether or not the join info can be used to join with a password
+	 */
+	virtual bool CanJoinWithPassword() const = 0;
+
+	/**
+	 * @return whether or not the join info has the info to request an invite
+	 */
+	virtual bool CanRequestAnInvite() const = 0;
+};
+
+/**
+ * Permissions for party features
+ */
+namespace PartySystemPermissions
+{
+	/**
+	 * EPresencePermissions details who can publish what to presence.
+	 * DoNotPublish is a valid setting for primary parties.
+	 */
+	enum class EPermissionType: uint8
+	{
+		None = 0x0,
+		Leader = 0x1,
+		Friend = 0x2,
+		Anyone = 0x4
+	};
+
+	/**
+	 * FPublishPermissionValidator uses meta programming to cause errors for invalid combinations of permission flags
+	 */
+	template<bool bIsValidPermissionCombo>
+		struct FPublishPermissionValidator;
+
+	template<> struct FPublishPermissionValidator<true> {};
+
+	template<EPermissionType PublishIdPermission, EPermissionType PublishKeyPermission>
+	struct FPublishPermissionBuilder : public FPublishPermissionValidator<(PublishKeyPermission <= PublishIdPermission)>
+	{
+		static const uint8 Value = static_cast<uint8>(PublishIdPermission) << 4 | static_cast<uint8>(PublishKeyPermission);
+	};
+
+	/**
+	 * Pre-defined presence permission objects
+	 */
+	enum class EPresencePermissions: uint8
+	{
+		DoNotPublish                    = FPublishPermissionBuilder<EPermissionType::None,   EPermissionType::None>::Value,
+		LeaderPublishIdNonePublishKey   = FPublishPermissionBuilder<EPermissionType::Leader, EPermissionType::None>::Value,
+		LeaderPublishIdLeaderPublishKey = FPublishPermissionBuilder<EPermissionType::Leader, EPermissionType::Leader>::Value,
+		FriendPublishIdNonePublishKey   = FPublishPermissionBuilder<EPermissionType::Friend, EPermissionType::None>::Value,
+		FriendPublishIdLeaderPublishKey = FPublishPermissionBuilder<EPermissionType::Friend, EPermissionType::Leader>::Value,
+		FriendPublishIdFriendPublishKey = FPublishPermissionBuilder<EPermissionType::Friend, EPermissionType::Friend>::Value,
+		AnyonePublishIdNonePublishKey   = FPublishPermissionBuilder<EPermissionType::Anyone, EPermissionType::None>::Value,
+		AnyonePublishIdLeaderPublishKey = FPublishPermissionBuilder<EPermissionType::Anyone, EPermissionType::Leader>::Value,
+		AnyonePublishIdFriendPublishKey = FPublishPermissionBuilder<EPermissionType::Anyone, EPermissionType::Friend>::Value,
+		AnyonePublishIdAnyonePublishKey = FPublishPermissionBuilder<EPermissionType::Anyone, EPermissionType::Anyone>::Value,
+	};
+
+	/**
+	 * Presence permission aliases
+	 */
+	static const EPresencePermissions FriendsInviteOnly = EPresencePermissions::LeaderPublishIdNonePublishKey;
+	static const EPresencePermissions FriendsOfFriendsInviteOnly = EPresencePermissions::FriendPublishIdNonePublishKey;
+	static const EPresencePermissions PublicInviteOnly = EPresencePermissions::AnyonePublishIdNonePublishKey;
+	static const EPresencePermissions FriendsOnly = EPresencePermissions::LeaderPublishIdLeaderPublishKey;
+	static const EPresencePermissions FriendsOfFriendsOnly = EPresencePermissions::FriendPublishIdFriendPublishKey;
+	static const EPresencePermissions Public = EPresencePermissions::AnyonePublishIdAnyonePublishKey;
+
+	enum class EInvitePermissions
+	{
+		/** Available to the leader only */
+		Leader,
+		/** Available to friends of the leader only */
+		Friends,
+		/** Available to anyone */
+		Anyone,
+		/** Noone can send invites */
+		Noone, 
+	};
+}
+
+enum class EJoinRequestAction
+{
+	Manual,
+	AutoApprove,
+	AutoReject
+};
+/**
+ * Options for configuring a new party or for updating an existing party
+ */
+struct FPartyConfiguration : public TSharedFromThis<FPartyConfiguration>
+{
+	FPartyConfiguration()
+		: JoinRequestAction(EJoinRequestAction::Manual)
+		, PresencePermissions(PartySystemPermissions::EPresencePermissions::AnyonePublishIdAnyonePublishKey)
+		, InvitePermissions(PartySystemPermissions::EInvitePermissions::Leader)
+		, bChatEnabled(true)
+		, bShouldRemoveOnDisconnection(false)
+		, bIsAcceptingMembers(false)
+		, NotAcceptingMembersReason(0)
+		, MaxMembers(0)
+	{}
+
+	/** should publish info to presence */
+	EJoinRequestAction JoinRequestAction;
+	/** Permission for how the party can be  */
+	PartySystemPermissions::EPresencePermissions PresencePermissions;
+	/** Permission who can send invites */
+	PartySystemPermissions::EInvitePermissions InvitePermissions;
+	/** should have a muc room */
+	bool bChatEnabled;
+	/** should remove on disconnection */
+	bool bShouldRemoveOnDisconnection;
+	/** is accepting members */
+	bool bIsAcceptingMembers;
+	/** not accepting members reason */
+	int32 NotAcceptingMembersReason;
+	/** Maximum active members allowed. 0 means no maximum. */
+	int32 MaxMembers;
+	/** Human readable nickname */
+	FString Nickname;
+	/** Human readable description */
+	FString Description;
+	/** Human readable password for party. */
+	FString Password;
+	/** clients can add whatever data they want for configuration options */
+	FOnlinePartyData ClientConfigData;
+};
+
+enum class EPartyState
+{
+	None,
+	CreatePending,
+	JoinPending,
+	LeavePending,
+	Active,
+	Disconnected,
+	Reconnecting,
+	CleanUp
+};
+
+/**
+ * Current state associated with a party
+ */
+class FOnlineParty : public TSharedFromThis<FOnlineParty>
+{
+protected:
+	FOnlineParty();
+	explicit FOnlineParty(const TSharedRef<const FOnlinePartyId>& InPartyId, const FOnlinePartyTypeId InPartyTypeId)
+		: PartyId(InPartyId)
+		, PartyTypeId(InPartyTypeId)
+		, State(EPartyState::None)
+		, Config(MakeShareable(new FPartyConfiguration()))
+	{}
+
+public:
+	virtual ~FOnlineParty()
+	{}
+
+	virtual bool CanLocalUserInvite(const FUniqueNetId& LocalUserId) const = 0;
+	virtual bool IsJoinable() const = 0;
+
+	/** unique id of the party */
+	TSharedRef<const FOnlinePartyId> PartyId;
+	/** unique id of the party */
+	const FOnlinePartyTypeId PartyTypeId;
+	/** unique id of the leader */
+	TSharedPtr<const FUniqueNetId> LeaderId;
+	/** The current state of the party */
+	EPartyState State;
+	/** Current state of configuration */
+	TSharedRef<FPartyConfiguration> Config;
+	/** id of chat room associated with the party */
+	FChatRoomId RoomId;
+};
+
+enum class EMemberChangedReason
+{
+	Disconnected,
+	Rejoined,
+	Promoted
+};
+
+enum class EMemberExitedReason
+{
+	Unknown,
+	Left,
+	Removed,
+	Kicked
+};
+
+struct FPartyInvitationRecipient
+{
+	FPartyInvitationRecipient(const TSharedRef<const FUniqueNetId>& InId)
+		: Id(InId)
+		{}
+
+	FPartyInvitationRecipient(const FUniqueNetId& InId)
+		: Id(InId.AsShared())
+	{}
+
+
+	TSharedRef<const FUniqueNetId> Id;
+	FString PlatformData;
+};
+
+
+enum class ECreatePartyCompletionResult;
+enum class EJoinPartyCompletionResult;
+enum class ELeavePartyCompletionResult;
+enum class EUpdateConfigCompletionResult;
+enum class ERequestPartyInvitationCompletionResult;
+enum class ESendPartyInvitationCompletionResult;
+enum class EAcceptPartyInvitationCompletionResult;
+enum class ERejectPartyInvitationCompletionResult;
+enum class EKickMemberCompletionResult;
+enum class EPromoteMemberCompletionResult;
+enum class EInvitationResponse;
+
+///////////////////////////////////////////////////////////////////
+// Completion delegates
+///////////////////////////////////////////////////////////////////
+/**
+ * Party creation async task completed callback
+ *
+ * @param LocalUserId - id of user that initiated the request
+ * @param PartyId - id associated with the party
+ * @param Result - result of the operation
+ */
+DECLARE_DELEGATE_ThreeParams(FOnCreatePartyComplete, const FUniqueNetId& /*LocalUserId*/, const TSharedPtr<const FOnlinePartyId>& /*PartyId*/, const ECreatePartyCompletionResult /*Result*/);
+/**
+ * Party join async task completed callback
+ *
+ * @param LocalUserId - id of user that initiated the request
+ * @param PartyId - id associated with the party
+ * @param Result - result of the operation
+ * @param NotApprovedReason - client defined value describing why you were not approved
+ */
+DECLARE_DELEGATE_FourParams(FOnJoinPartyComplete, const FUniqueNetId& /*LocalUserId*/, const FOnlinePartyId& /*PartyId*/, const EJoinPartyCompletionResult /*Result*/, const int32 /*NotApprovedReason*/);
+/**
+ * Party leave async task completed callback
+ *
+ * @param LocalUserId - id of user that initiated the request
+ * @param PartyId - id associated with the party
+ * @param Result - result of the operation
+ */
+DECLARE_DELEGATE_ThreeParams(FOnLeavePartyComplete, const FUniqueNetId& /*LocalUserId*/, const FOnlinePartyId& /*PartyId*/, const ELeavePartyCompletionResult /*Result*/);
+/**
+ * Party update async task completed callback
+ *
+ * @param LocalUserId - id of user that initiated the request
+ * @param PartyId - id associated with the party
+ * @param Result - result of the operation
+ */
+DECLARE_DELEGATE_ThreeParams(FOnUpdatePartyComplete, const FUniqueNetId& /*LocalUserId*/, const FOnlinePartyId& /*PartyId*/, const EUpdateConfigCompletionResult /*Result*/);
+/**
+ * Party update async task completed callback
+ *
+ * @param LocalUserId - id of user that initiated the request
+ * @param PartyId - id associated with the party
+ * @param Result - result of the operation
+ */
+DECLARE_DELEGATE_ThreeParams(FOnRequestPartyInvitationComplete, const FUniqueNetId& /*LocalUserId*/, const FOnlinePartyId& /*PartyId*/, const ERequestPartyInvitationCompletionResult /*Result*/);
+/**
+ * Party invitation sent completed callback
+ *
+ * @param LocalUserId - id of user that initiated the request
+ * @param PartyId - id associated with the party
+ * @param RecipientId - user invite was sent to
+ * @param Result - result of the send operation
+ */
+DECLARE_DELEGATE_FourParams(FOnSendPartyInvitationComplete, const FUniqueNetId& /*LocalUserId*/, const FOnlinePartyId& /*PartyId*/, const FUniqueNetId& /*RecipientId*/, const ESendPartyInvitationCompletionResult /*Result*/);
+/**
+ * Accepting an invite to a user to join party async task completed callback
+ *
+ * @param LocalUserId - id of user that initiated the request
+ * @param PartyId - id associated with the party
+ * @param Result - string with error info if any
+ */
+DECLARE_DELEGATE_ThreeParams(FOnAcceptPartyInvitationComplete, const FUniqueNetId& /*LocalUserId*/, const FOnlinePartyId& /*PartyId*/, const EAcceptPartyInvitationCompletionResult /*Result*/);
+/**
+ * Rejecting an invite to a user to join party async task completed callback
+ *
+ * @param LocalUserId - id of user that initiated the request
+ * @param bWasSuccessful - true if successfully sent invite
+ * @param PartyId - id associated with the party
+ * @param Result - string with error info if any
+ */
+DECLARE_DELEGATE_ThreeParams(FOnRejectPartyInvitationComplete, const FUniqueNetId& /*LocalUserId*/, const FOnlinePartyId& /*PartyId*/, const ERejectPartyInvitationCompletionResult /*Result*/);
+/**
+ * Kicking a member of a party async task completed callback
+ *
+ * @param LocalUserId - id of user that initiated the request
+ * @param PartyId - id associated with the party
+ * @param MemberId - id of member being kicked
+ * @param Result - string with error info if any
+ */
+DECLARE_DELEGATE_FourParams(FOnKickPartyMemberComplete, const FUniqueNetId& /*LocalUserId*/, const FOnlinePartyId& /*PartyId*/, const FUniqueNetId& /*LocalUserId*/, const EKickMemberCompletionResult /*Result*/);
+/**
+ * Promoting a member of a party async task completed callback
+ *
+ * @param LocalUserId - id of user that initiated the request
+ * @param PartyId - id associated with the party
+ * @param MemberId - id of member being promoted to leader
+ * @param Result - string with error info if any
+ */
+DECLARE_DELEGATE_FourParams(FOnPromotePartyMemberComplete, const FUniqueNetId& /*LocalUserId*/, const FOnlinePartyId& /*PartyId*/, const FUniqueNetId& /*LocalUserId*/, const EPromoteMemberCompletionResult /*Result*/);
+
+
+
+
+
+
+///////////////////////////////////////////////////////////////////
+// Notification delegates
+///////////////////////////////////////////////////////////////////
+
+/**
+* notification when a party is joined
+* @param LocalUserId - id associated with this notification
+* @param PartyId - id associated with the party
+*/
+DECLARE_MULTICAST_DELEGATE_TwoParams(F_PREFIX(OnPartyJoined), const FUniqueNetId& /*LocalUserId*/, const FOnlinePartyId& /*PartyId*/);
+PARTY_DECLARE_DELEGATETYPE(OnPartyJoined);
+
+/**
+ * notification when a party is joined
+ * @param LocalUserId - id associated with this notification
+ * @param PartyId - id associated with the party
+ */
+DECLARE_MULTICAST_DELEGATE_TwoParams(F_PREFIX(OnPartyExited), const FUniqueNetId& /*LocalUserId*/, const FOnlinePartyId& /*PartyId*/);
+PARTY_DECLARE_DELEGATETYPE(OnPartyExited);
+
+/**
+ * Notification when player promotion is locked out.
+ * @param LocalUserId - id associated with this notification
+ * @param PartyId - id associated with the party
+ * @param bLockoutState - if promotion is currently locked out
+ */
+DECLARE_MULTICAST_DELEGATE_ThreeParams(F_PREFIX(OnPartyPromotionLockoutChanged), const FUniqueNetId& /*LocalUserId*/, const FOnlinePartyId& /*PartyId*/, const bool /*bLockoutState*/);
+PARTY_DECLARE_DELEGATETYPE(OnPartyPromotionLockoutChanged);
+
+/**
+ * Notification when party data is updated
+ * @param LocalUserId - id associated with this notification
+ * @param PartyId - id associated with the party
+ * @param PartyConfig - party whose config was updated
+ */
+DECLARE_MULTICAST_DELEGATE_ThreeParams(F_PREFIX(OnPartyConfigChanged), const FUniqueNetId& /*LocalUserId*/, const FOnlinePartyId& /*PartyId*/, const TSharedRef<FPartyConfiguration>& /*PartyConfig*/);
+PARTY_DECLARE_DELEGATETYPE(OnPartyConfigChanged);
+
+/**
+ * Notification when party data is updated
+ * @param LocalUserId - id associated with this notification
+ * @param PartyId - id associated with the party
+ * @param PartyData - party data that was updated
+ */
+DECLARE_MULTICAST_DELEGATE_ThreeParams(F_PREFIX(OnPartyDataReceived), const FUniqueNetId& /*LocalUserId*/, const FOnlinePartyId& /*PartyId*/, const TSharedRef<FOnlinePartyData>& /*PartyData*/);
+PARTY_DECLARE_DELEGATETYPE(OnPartyDataReceived);
+
+/**
+ * Notification when a member changes in a party
+ * @param LocalUserId - id associated with this notification
+ * @param PartyId - id associated with the party
+ * @param MemberId - id of member that joined
+ * @param Reason - how the member changed
+ */
+DECLARE_MULTICAST_DELEGATE_FourParams(F_PREFIX(OnPartyMemberChanged), const FUniqueNetId& /*LocalUserId*/, const FOnlinePartyId& /*PartyId*/, const FUniqueNetId& /*MemberId*/, const EMemberChangedReason /*Reason*/);
+PARTY_DECLARE_DELEGATETYPE(OnPartyMemberChanged);
+
+/**
+ * Notification when a member exits a party
+ * @param LocalUserId - id associated with this notification
+ * @param PartyId - id associated with the party
+ * @param MemberId - id of member that joined
+ * @param Reason - why the member was removed
+ */
+DECLARE_MULTICAST_DELEGATE_FourParams(F_PREFIX(OnPartyMemberExited), const FUniqueNetId& /*LocalUserId*/, const FOnlinePartyId& /*PartyId*/, const FUniqueNetId& /*MemberId*/, const EMemberExitedReason /*Reason*/);
+PARTY_DECLARE_DELEGATETYPE(OnPartyMemberExited);
+
+/**
+ * Notification when a member joins the party
+ * @param LocalUserId - id associated with this notification
+ * @param PartyId - id associated with the party
+ * @param MemberId - id of member that joined
+ */
+DECLARE_MULTICAST_DELEGATE_ThreeParams(F_PREFIX(OnPartyMemberJoined), const FUniqueNetId& /*LocalUserId*/, const FOnlinePartyId& /*PartyId*/, const FUniqueNetId& /*MemberId*/);
+PARTY_DECLARE_DELEGATETYPE(OnPartyMemberJoined);
+
+/**
+ * Notification when party member data is updated
+ * @param LocalUserId - id associated with this notification
+ * @param PartyId - id associated with the party
+ * @param MemberId - id of member that had updated data
+ * @param PartyMemberData - party member data that was updated
+ */
+DECLARE_MULTICAST_DELEGATE_FourParams(F_PREFIX(OnPartyMemberDataReceived), const FUniqueNetId& /*LocalUserId*/, const FOnlinePartyId& /*PartyId*/, const FUniqueNetId& /*MemberId*/, const TSharedRef<FOnlinePartyData>& /*PartyMemberData*/);
+PARTY_DECLARE_DELEGATETYPE(OnPartyMemberDataReceived);
+
+/**
+ * Notification when an invite list has changed for a party
+ * @param LocalUserId - user that is associated with this notification
+ */
+DECLARE_MULTICAST_DELEGATE_OneParam(F_PREFIX(OnPartyInvitesChanged), const FUniqueNetId& /*LocalUserId*/);
+PARTY_DECLARE_DELEGATETYPE(OnPartyInvitesChanged);
+
+/**
+ * Notification when a request for an invite has been received
+ * @param LocalUserId - id associated with this notification
+ * @param PartyId - id associated with the party
+ * @param SenderId - id of user that sent the invite
+ * @param RequestForId - id of user that sender is requesting the invite for - invalid if the sender is requesting the invite
+ */
+DECLARE_MULTICAST_DELEGATE_FourParams(F_PREFIX(OnPartyInviteRequestReceived), const FUniqueNetId& /*LocalUserId*/, const FOnlinePartyId& /*PartyId*/, const FUniqueNetId& /*SenderId*/, const FUniqueNetId& /*RequestForId*/);
+PARTY_DECLARE_DELEGATETYPE(OnPartyInviteRequestReceived);
+
+/**
+ * Notification when a new invite is received
+ * @param LocalUserId - id associated with this notification
+ * @param PartyId - id associated with the party
+ * @param SenderId - id of member that sent the invite
+ */
+DECLARE_MULTICAST_DELEGATE_ThreeParams(F_PREFIX(OnPartyInviteReceived), const FUniqueNetId& /*LocalUserId*/, const FOnlinePartyId& /*PartyId*/, const FUniqueNetId& /*SenderId*/);
+PARTY_DECLARE_DELEGATETYPE(OnPartyInviteReceived);
+
+/**
+ * Notification when a new invite is received
+ * @param LocalUserId - id associated with this notification
+ * @param PartyId - id associated with the party
+ * @param SenderId - id of member that sent the invite
+ * @param bWasAccepted - whether or not the invite was accepted
+ */
+DECLARE_MULTICAST_DELEGATE_FourParams(F_PREFIX(OnPartyInviteResponseReceived), const FUniqueNetId& /*LocalUserId*/, const FOnlinePartyId& /*PartyId*/, const FUniqueNetId& /*SenderId*/, const EInvitationResponse /*Response*/);
+PARTY_DECLARE_DELEGATETYPE(OnPartyInviteResponseReceived);
+
+/**
+ * Notification when a new reservation request is received
+ * @param LocalUserId - id associated with this notification
+ * @param PartyId - id associated with the party
+ * @param SenderId - id of member that sent the request
+ */
+DECLARE_MULTICAST_DELEGATE_ThreeParams(F_PREFIX(OnPartyJoinRequestReceived), const FUniqueNetId& /*LocalUserId*/, const FOnlinePartyId& /*PartyId*/, const FUniqueNetId& /*SenderId*/);
+PARTY_DECLARE_DELEGATETYPE(OnPartyJoinRequestReceived);
+
+/**
+ * Notification when an invite is accepted
+ * @param LocalUserId - id associated with this notification
+ * @param PartyId - id associated with the party
+ * @param MemberId - id of member that accepted the invite
+ * @param bWasAccepted - whether or not the invite was accepted
+ */
+DECLARE_MULTICAST_DELEGATE_FourParams(F_PREFIX(OnPartyJoinRequestResponseReceived), const FUniqueNetId& /*LocalUserId*/, const FOnlinePartyId& /*PartyId*/, const FUniqueNetId& /*MemberId*/, bool /*bWasAccepted*/);
+PARTY_DECLARE_DELEGATETYPE(OnPartyJoinRequestResponseReceived);
+
+/**
+ * Interface definition for the online party services 
+ * Allows for forming a party and communicating with party members
+ */
+class IOnlinePartySystem
+{
+protected:
+	IOnlinePartySystem() {};
+
+public:
+	virtual ~IOnlinePartySystem() {};
+
+	/**
+	 * Create a new party
+	 *
+	 * @param LocalUserId - user making the request
+	 * @param PartyConfig - configuration for the party (can be updated later)
+	 * @param Delegate - called on completion
+	 * @param UserRoomId - this forces the name of the room to be this value
+	 *
+	 * @return true if task was started
+	 */
+	virtual bool CreateParty(const FUniqueNetId& LocalUserId, const FOnlinePartyTypeId PartyTypeId, const FPartyConfiguration& PartyConfig, const FOnCreatePartyComplete& Delegate = FOnCreatePartyComplete()) = 0;
+
+	/**
+	 * Update an existing party with new configuration
+	 *
+	 * @param LocalUserId - user making the request
+	 * @param PartyConfig - configuration for the party
+	 * @param Delegate - called on completion
+	 *
+	 * @return true if task was started
+	 */
+	virtual bool UpdateParty(const FUniqueNetId& LocalUserId, const FOnlinePartyId& PartyId, const FPartyConfiguration& PartyConfig, bool bShouldRegenerateReservationKey = false, const FOnUpdatePartyComplete& Delegate = FOnUpdatePartyComplete()) = 0;
+
+	/**
+	 * Join an existing party
+	 *
+	 * @param LocalUserId - user making the request
+	 * @param PartyId - id of an existing party
+	 * @param RecipientId - id of the party leader
+	 * @param ReservationKey - the ReservationKey for the party
+	 * @param Delegate - called on completion
+	 *
+	 * @return true if task was started
+	 */
+	virtual bool JoinParty(const FUniqueNetId& LocalUserId, const IOnlinePartyJoinInfo& OnlinePartyJoinInfo, const FOnJoinPartyComplete& Delegate = FOnJoinPartyComplete()) = 0;
+
+	/**
+	 * Leave an existing party
+	 * All existing party members notified of member leaving (see FOnPartyMemberLeft)
+	 *
+	 * @param LocalUserId - user making the request
+	 * @param PartyId - id of an existing party
+	 * @param Delegate - called on completion
+	 *
+	 * @return true if task was started
+	 */
+	virtual bool LeaveParty(const FUniqueNetId& LocalUserId, const FOnlinePartyId& PartyId, const FOnLeavePartyComplete& Delegate = FOnLeavePartyComplete()) = 0;
+
+	/**
+	* Approve a request to join a party
+	*
+	* @param LocalUserId - user making the request
+	* @param PartyId - id of an existing party
+	* @param RecipientId - id of the user being invited
+	* @param DeniedResultCode - client defined value to return when leader denies approval
+	*
+	* @return true if task was started
+	*/
+	virtual bool ApproveJoinRequest(const FUniqueNetId& LocalUserId, const FOnlinePartyId& PartyId, const FUniqueNetId& RecipientId, bool bIsApproved, int32 DeniedResultCode = 0) = 0;
+
+	/**
+	 * sends an invitation to a user that could not otherwise join a party
+	 * if the player accepts the invite they will be sent the data needed to trigger a call to RequestReservation
+	 *
+	 * @param LocalUserId - user making the request
+	 * @param PartyId - id of an existing party
+	 * @param Recipient - structure specifying the recipient of the invitation
+	 * @param Delegate - called on completion
+	 *
+	 * @return true if task was started
+	 */
+	virtual bool SendInvitation(const FUniqueNetId& LocalUserId, const FOnlinePartyId& PartyId, const FPartyInvitationRecipient& Recipient, const FOnlinePartyData& ClientData = FOnlinePartyData(), const FOnSendPartyInvitationComplete& Delegate = FOnSendPartyInvitationComplete()) = 0;
+
+	/**
+	* Accept an invite to a party. NOTE this does not initiate a join.
+	*
+	* @param LocalUserId - user making the request
+	* @param PartyId - id of an existing party
+	* @param Delegate - called on completion
+	*
+	* @return true if task was started
+	*/
+	virtual bool AcceptInvitation(const FUniqueNetId& LocalUserId, const FUniqueNetId& SenderId) = 0;
+
+	/**
+	* Reject an invite to a party
+	*
+	* @param LocalUserId - user making the request
+	* @param PartyId - id of an existing party
+	*
+	* @return true if task was started
+	*/
+	virtual bool RejectInvitation(const FUniqueNetId& LocalUserId, const FUniqueNetId& SenderId) = 0;
+
+	/**
+	 * Kick a user from an existing party
+	 * Only admin can kick a party member
+	 *
+	 * @param LocalUserId - user making the request
+	 * @param PartyId - id of an existing party
+	 * @param MemberId - id of the user being kicked
+	 * @param Delegate - called on completion
+	 *
+	 * @return true if task was started
+	 */
+	virtual bool KickMember(const FUniqueNetId& LocalUserId, const FOnlinePartyId& PartyId, const FUniqueNetId& TargetMemberId, const FOnKickPartyMemberComplete& Delegate = FOnKickPartyMemberComplete()) = 0;
+
+	/**
+	 * Promote a user from an existing party to be admin
+	 * All existing party members notified of promoted member (see FOnPartyMemberPromoted)
+	 *
+	 * @param LocalUserId - user making the request
+	 * @param PartyId - id of an existing party
+	 * @param MemberId - id of the user being promoted
+	 * @param Delegate - called on completion
+	 *
+	 * @return true if task was started
+	 */
+	virtual bool PromoteMember(const FUniqueNetId& LocalUserId, const FOnlinePartyId& PartyId, const FUniqueNetId& TargetMemberId, const FOnPromotePartyMemberComplete& Delegate = FOnPromotePartyMemberComplete()) = 0;
+
+	/**
+	 * Set party data and broadcast to all members
+	 * Only current data can be set and no history of past party data is preserved
+	 * Party members notified of new data (see FOnPartyDataReceived)
+	 *
+	 * @param LocalUserId - user making the request
+	 * @param PartyId - id of an existing party
+	 * @param PartyData - data to send to all party members
+	 * @param Delegate - called on completion
+	 *
+	 * @return true if task was started
+	 */
+	virtual bool UpdatePartyData(const FUniqueNetId& LocalUserId, const FOnlinePartyId& PartyId, const FOnlinePartyData& PartyData) = 0;
+
+	/**
+	 * Set party data for a single party member and broadcast to all members
+	 * Only current data can be set and no history of past party member data is preserved
+	 * Party members notified of new data (see FOnPartyMemberDataReceived)
+	 *
+	 * @param LocalUserId - user making the request
+	 * @param PartyId - id of an existing party
+	 * @param PartyMemberData - member data to send to all party members
+	 * @param Delegate - called on completion
+	 *
+	 * @return true if task was started
+	 */
+	virtual bool UpdatePartyMemberData(const FUniqueNetId& LocalUserId, const FOnlinePartyId& PartyId, const FOnlinePartyData& PartyMemberData) = 0;
+
+	/**
+	 * returns true if the user specified by MemberId is the leader of the party specified by PartyId
+	 *
+	 * @param LocalUserId - user making the request
+	 * @param PartyId     - id of an existing party
+	 * @param MemberId    - id of member to test
+	 *
+	 */
+	virtual bool IsMemberLeader(const FUniqueNetId& LocalUserId, const FOnlinePartyId& PartyId, const FUniqueNetId& MemberId) const = 0;
+
+	/**
+	 * returns the number of players in a given party
+	 *
+	 * @param LocalUserId - user making the request
+	 * @param PartyId     - id of an existing party
+	 *
+	 */
+	virtual uint32 GetPartyMemberCount(const FUniqueNetId& LocalUserId, const FOnlinePartyId& PartyId) const = 0;
+
+	/**
+	 * Get info associated with a party
+	 *
+	 * @param LocalUserId - user making the request
+	 * @param PartyId     - id of an existing party
+	 *
+	 * @return party info or nullptr if not found
+	 */
+	virtual TSharedPtr<const FOnlineParty> GetParty(const FUniqueNetId& LocalUserId, const FOnlinePartyId& PartyId) const = 0;
+
+	/**
+	 * Get info associated with a party
+	 *
+	 * @param LocalUserId - user making the request
+	 * @param PartyTypeId - type of an existing party
+	 *
+	 * @return party info or nullptr if not found
+	 */
+	virtual TSharedPtr<const FOnlineParty> GetParty(const FUniqueNetId& LocalUserId, const FOnlinePartyTypeId& PartyTypeId) const = 0;
+
+	/**
+	 * Get a party member by id
+	 *
+	 * @param LocalUserId - user making the request
+	 * @param PartyId     - id of an existing party
+	 * @param MemberId    - id of member to find
+	 *
+	 * @return party member info or nullptr if not found
+	 */
+	virtual TSharedPtr<FOnlinePartyMember> GetPartyMember(const FUniqueNetId& LocalUserId, const FOnlinePartyId& PartyId, const FUniqueNetId& MemberId) const = 0;
+
+	/**
+	 * Get current cached data associated with a party
+	 * FOnPartyDataReceived notification called whenever this data changes
+	 *
+	 * @param LocalUserId - user making the request
+	 * @param PartyId     - id of an existing party
+	 *
+	 * @return party data or nullptr if not found
+	 */
+	virtual TSharedPtr<FOnlinePartyData> GetPartyData(const FUniqueNetId& LocalUserId, const FOnlinePartyId& PartyId) const = 0;
+
+	/**
+	 * Get current cached data associated with a party member
+	 * FOnPartyMemberDataReceived notification called whenever this data changes
+	 *
+	 * @param LocalUserId - user making the request
+	 * @param PartyId     - id of an existing party
+	 * @param MemberId    - id of member to find data for
+	 *
+	 * @return party member data or nullptr if not found
+	 */
+	virtual TSharedPtr<FOnlinePartyData> GetPartyMemberData(const FUniqueNetId& LocalUserId, const FOnlinePartyId& PartyId, const FUniqueNetId& MemberId) const = 0;
+
+	/**
+	 * returns true if the user is advertising a party of that type on their presence
+	 *
+	 * @param LocalUserId       - user making the request
+	 * @param UserId            - user to check
+	 * @param PartyTypeId       - type of party to query
+	 *
+	 */
+	virtual TSharedPtr<IOnlinePartyJoinInfo> GetAdvertisedParty(const FUniqueNetId& LocalUserId, const FUniqueNetId& UserId, const FOnlinePartyTypeId PartyTypeId) const = 0;
+
+	/**
+	 * Get a list of currently joined parties for the user
+	 *
+	 * @param LocalUserId     - user making the request
+	 * @param OutPartyIdArray - list of party ids joined by the current user
+	 *
+	 * @return true if entries found
+	 */
+	virtual bool GetJoinedParties(const FUniqueNetId& LocalUserId, TArray<TSharedRef<const FOnlinePartyId>>& OutPartyIdArray) const = 0;
+
+	/**
+	 * Get list of current party members
+	 *
+	 * @param LocalUserId          - user making the request
+	 * @param PartyId              - id of an existing party
+	 * @param OutPartyMembersArray - list of party members currently in the party
+	 *
+	 * @return true if entries found
+	 */
+	virtual bool GetPartyMembers(const FUniqueNetId& LocalUserId, const FOnlinePartyId& PartyId, TArray<TSharedRef<FOnlinePartyMember>>& OutPartyMembersArray) const = 0;
+
+	/**
+	 * Get a list of parties the user has been invited to
+	 *
+	 * @param LocalUserId            - user making the request
+	 * @param OutPendingInvitesArray - list of party info needed to join the party
+	 *
+	 * @return true if entries found
+	 */
+	virtual bool GetPendingInvites(const FUniqueNetId& LocalUserId, TArray<TSharedRef<IOnlinePartyJoinInfo>>& OutPendingInvitesArray) const = 0;
+
+	/**
+	 * Get list of users requesting to join the party
+	 *
+	 * @param LocalUserId           - user making the request
+	 * @param PartyId               - id of an existing party
+	 * @param OutPendingUserIdArray - list of pending party members
+	 *
+	 * @return true if entries found
+	 */
+	virtual bool GetPendingJoinRequests(const FUniqueNetId& LocalUserId, const FOnlinePartyId& PartyId, TArray<TSharedRef<IOnlinePartyPendingJoinRequestInfo>>& OutPendingJoinRequestArray) const = 0;
+
+	/**
+	 * Get list of users invited to a party that have not yet responded
+	 *
+	 * @param LocalUserId                - user making the request
+	 * @param PartyId                    - id of an existing party
+	 * @param OutPendingInvitedUserArray - list of user that have pending invites
+	 *
+	 * @return true if entries found
+	 */
+	virtual bool GetPendingInvitedUsers(const FUniqueNetId& LocalUserId, const FOnlinePartyId& PartyId, TArray<TSharedRef<const FUniqueNetId>>& OutPendingInvitedUserArray) const = 0;
+
+	static const FOnlinePartyTypeId::TInternalType PrimaryPartyTypeIdValue = 0x11111111;
+	/**
+	 * @return party type id for the primary party - the primary party is the party that will be addressable via the social panel
+	 */
+	static const FOnlinePartyTypeId GetPrimaryPartyTypeId() { return FOnlinePartyTypeId(PrimaryPartyTypeIdValue); }
+
+	/**
+	 * @return party type id for the primary party - the primary party is the party that will be addressable via the social panel
+	 */
+	static const FOnlinePartyTypeId MakePartyTypeId(const FOnlinePartyTypeId::TInternalType InTypeId) { ensure(InTypeId != PrimaryPartyTypeIdValue); return FOnlinePartyTypeId(InTypeId); }
+
+	/**
+	 * returns the json version of a join info for a current party
+	 *
+	 * @param LocalUserId       - user making the request
+	 * @param PartyId           - party to make the json from
+	 *
+ 	 */
+	virtual FString MakeJoinInfoJson(const FUniqueNetId& LocalUserId, const FOnlinePartyId& PartyId) = 0;
+
+	/**
+	 * returns a valid join info object from a json blob
+	 *
+	 * @param JoinInfoJson       - json blob to convert
+	 *
+	 */
+	virtual TSharedPtr<IOnlinePartyJoinInfo> MakeJoinInfoFromJson(const FString& JoinInfoJson) = 0;
+
+	/**
+	 * Creates a command line token from a IOnlinePartyJoinInfo object
+	 *
+	 * @param JoinInfo - the IOnlinePartyJoinInfo object to convert
+	 *
+	 * return the new IOnlinePartyJoinInfo object
+	 */
+	virtual FString MakeTokenFromJoinInfo(const IOnlinePartyJoinInfo& JoinInfo) const = 0;
+
+	/**
+	 * Creates a IOnlinePartyJoinInfo object from a command line token
+	 *
+	 * @param Token - the token string
+	 *
+	 * return the new IOnlinePartyJoinInfo object
+	 */
+	virtual TSharedRef<IOnlinePartyJoinInfo> MakeJoinInfoFromToken(const FString& Token) const = 0;
+
+	/**
+	 * Checks to see if there is a pending command line invite and consumes it
+	 *
+	 * return the pending IOnlinePartyJoinInfo object
+	 */
+	virtual TSharedPtr<IOnlinePartyJoinInfo> ConsumePendingCommandLineInvite() = 0;
+
+	/**
+	 * List of all subscribe-able notifications
+	 *
+	 * OnPartyJoined
+	 * OnPartyPromotionLockoutStateChanged
+	 * OnPartyConfigChanged
+	 * OnPartyDataChanged
+	 * OnPartyMemberChanged
+	 * OnPartyMemberExited
+	 * OnPartyMemberJoined
+	 * OnPartyMemberDataChanged
+	 * OnPartyInvitesChanged
+	 * OnPartyInviteRequestReceived
+	 * OnPartyInviteReceived
+	 * OnPartyInviteResponseReceived
+	 * OnPartyJoinRequestReceived
+	 * OnPartyJoinRequestResponseReceived
+	 *
+	 */
+
+	/**
+	 * notification of when a party is joined
+	 * @param LocalUserId - id associated with this notification
+	 * @param PartyId - id associated with the party
+	 */
+	DEFINE_ONLINE_DELEGATE_TWO_PARAM(OnPartyJoined, const FUniqueNetId& /*LocalUserId*/, const FOnlinePartyId& /*PartyId*/);
+
+	/**
+	 * notification of when a party is exited
+	 * @param LocalUserId - id associated with this notification
+	 * @param PartyId - id associated with the party
+	 */
+	DEFINE_ONLINE_DELEGATE_TWO_PARAM(OnPartyExited, const FUniqueNetId& /*LocalUserId*/, const FOnlinePartyId& /*PartyId*/);
+
+	/**
+	 * Notification when player promotion is locked out.
+	 *
+	 * @param PartyId - id associated with the party
+	 * @param bLockoutState - if promotion is currently locked out
+	 */
+	DEFINE_ONLINE_DELEGATE_THREE_PARAM(OnPartyPromotionLockoutChanged, const FUniqueNetId& /*LocalUserId*/, const FOnlinePartyId& /*PartyId*/, const bool /*bLockoutState*/);
+
+	/**
+	 * Notification when party data is updated
+	 * @param LocalUserId - id associated with this notification
+	 * @param PartyId - id associated with the party
+	 * @param PartyConfig - party whose config was updated
+	 */
+	DEFINE_ONLINE_DELEGATE_THREE_PARAM(OnPartyConfigChanged, const FUniqueNetId& /*LocalUserId*/, const FOnlinePartyId& /*PartyId*/, const TSharedRef<FPartyConfiguration>& /*PartyConfig*/);
+
+	/**
+	 * Notification when party data is updated
+	 * @param LocalUserId - id associated with this notification
+	 * @param PartyId - id associated with the party
+	 * @param PartyData - party data that was updated
+	 */
+	DEFINE_ONLINE_DELEGATE_THREE_PARAM(OnPartyDataReceived, const FUniqueNetId& /*LocalUserId*/, const FOnlinePartyId& /*PartyId*/, const TSharedRef<FOnlinePartyData>& /*PartyData*/);
+
+	/**
+	* Notification when a member changes in a party
+	* @param LocalUserId - id associated with this notification
+	* @param PartyId - id associated with the party
+	* @param MemberId - id of member that changed
+	* @param Reason - how the member changed
+	*/
+	DEFINE_ONLINE_DELEGATE_FOUR_PARAM(OnPartyMemberChanged, const FUniqueNetId& /*LocalUserId*/, const FOnlinePartyId& /*PartyId*/, const FUniqueNetId& /*MemberId*/, const EMemberChangedReason /*Reason*/);
+
+	/**
+	* Notification when a member exits a party
+	* @param LocalUserId - id associated with this notification
+	* @param PartyId - id associated with the party
+	* @param MemberId - id of member that joined
+	* @param Reason - why the member was removed
+	*/
+	DEFINE_ONLINE_DELEGATE_FOUR_PARAM(OnPartyMemberExited, const FUniqueNetId& /*LocalUserId*/, const FOnlinePartyId& /*PartyId*/, const FUniqueNetId& /*MemberId*/, const EMemberExitedReason /*Reason*/);
+
+	/**
+	 * Notification when a member joins the party
+	 * @param LocalUserId - id associated with this notification
+	 * @param PartyId - id associated with the party
+	 * @param MemberId - id of member that joined
+	 */
+	DEFINE_ONLINE_DELEGATE_THREE_PARAM(OnPartyMemberJoined, const FUniqueNetId& /*LocalUserId*/, const FOnlinePartyId& /*PartyId*/, const FUniqueNetId& /*MemberId*/);
+
+	/**
+	 * Notification when party member data is updated
+	 * @param LocalUserId - id associated with this notification
+	 * @param PartyId - id associated with the party
+	 * @param MemberId - id of member that had updated data
+	 * @param PartyMemberData - party member data that was updated
+	 */
+	DEFINE_ONLINE_DELEGATE_FOUR_PARAM(OnPartyMemberDataReceived, const FUniqueNetId& /*LocalUserId*/, const FOnlinePartyId& /*PartyId*/, const FUniqueNetId& /*MemberId*/, const TSharedRef<FOnlinePartyData>& /*PartyMemberData*/);
+
+	/**
+	 * Notification when an invite list has changed for a party
+	 * @param LocalUserId - user that is associated with this notification
+	 */
+	DEFINE_ONLINE_DELEGATE_ONE_PARAM(OnPartyInvitesChanged, const FUniqueNetId& /*LocalUserId*/);
+
+	/**
+	* Notification when a request for an invite has been received
+	* @param LocalUserId - id associated with this notification
+	* @param PartyId - id associated with the party
+	* @param SenderId - id of user that sent the invite
+	* @param RequestForId - id of user that sender is requesting the invite for - invalid if the sender is requesting the invite
+	*/
+	DEFINE_ONLINE_DELEGATE_FOUR_PARAM(OnPartyInviteRequestReceived, const FUniqueNetId& /*LocalUserId*/, const FOnlinePartyId& /*PartyId*/, const FUniqueNetId& /*SenderId*/, const FUniqueNetId& /*RequestForId*/);
+
+	/**
+	* Notification when a new invite is received
+	* @param LocalUserId - id associated with this notification
+	* @param SenderId - id of member that sent the invite
+	* @param PartyId - id associated with the party
+	*/
+	DEFINE_ONLINE_DELEGATE_THREE_PARAM(OnPartyInviteReceived, const FUniqueNetId& /*LocalUserId*/, const FOnlinePartyId& /*PartyId*/, const FUniqueNetId& /*SenderId*/);
+
+	/**
+	 * Notification when a new invite is received
+	 * @param LocalUserId - id associated with this notification
+	 * @param PartyId - id associated with the party
+	 * @param SenderId - id of member that sent the invite
+	 * @param bWasAccepted - true is the invite was accepted
+	 */
+	DEFINE_ONLINE_DELEGATE_FOUR_PARAM(OnPartyInviteResponseReceived, const FUniqueNetId& /*LocalUserId*/, const FOnlinePartyId& /*PartyId*/, const FUniqueNetId& /*SenderId*/, const EInvitationResponse /*Response*/);
+
+	/**
+	 * Notification when a new reservation request is received
+	 * @param LocalUserId - id associated with this notification
+	 * @param PartyId - id associated with the party
+	 * @param SenderId - id of member that sent the request
+	 */
+	DEFINE_ONLINE_DELEGATE_THREE_PARAM(OnPartyJoinRequestReceived, const FUniqueNetId& /*LocalUserId*/, const FOnlinePartyId& /*PartyId*/, const FUniqueNetId& /*SenderId*/);
+
+	/**
+	 * Notification when a join request is approved
+	 * @param LocalUserId - id associated with this notification
+	 * @param PartyId - id associated with the party
+	 * @param MemberId - id of member that accepted the invite
+	 * @param bWasAccepted - whether or not the invite was accepted
+	 */
+	DEFINE_ONLINE_DELEGATE_FOUR_PARAM(OnPartyJoinRequestResponseReceived, const FUniqueNetId& /*LocalUserId*/, const FOnlinePartyId& /*PartyId*/, const FUniqueNetId& /*MemberId*/, bool /*bWasAccepted*/);
+
+	/**
+	 * Dump out party state for all known parties
+	 */
+	virtual void DumpPartyState() = 0;
+
+};
+
+enum class ECreatePartyCompletionResult
+{
+	UnknownClientFailure = -100,
+	AlreadyInPartyOfSpecifiedType,
+	AlreadyCreatingParty,
+	AlreadyInParty,
+	FailedToCreateMucRoom,
+	NoResponse,
+	UnknownInternalFailure = 0,
+	Succeeded = 1
+};
+
+enum class EJoinPartyCompletionResult
+{
+	UnknownClientFailure = -100,
+	BadBuild,
+	JoinInfoInvalid,
+	AlreadyInPartyOfSpecifiedType,
+	AlreadyJoiningParty,
+	AlreadyInParty,
+	MessagingFailure,
+	PartyNotInPendingState,
+	ResponseFromUnexpectedUser,
+	NoSpace,
+	NotApproved,
+	RequesteeNotMember,
+	RequesteeNotLeader,
+	NoResponse,
+	UnknownInternalFailure = 0,
+	Succeeded = 1
+};
+
+enum class ELeavePartyCompletionResult
+{
+	UnknownClientFailure = -100,
+	LeavePending,
+	UnknownLocalUser,
+	UnknownParty,
+	NotMember,
+	MessagingFailure,
+	NoResponse,
+	UnknownTransportFailure,
+	UnknownInternalFailure = 0,
+	Succeeded = 1
+};
+
+enum class EUpdateConfigCompletionResult
+{
+	UnknownClientFailure = -100,
+	UnknownParty,
+	LocalMemberNotMember,
+	LocalMemberNotLeader,
+	RemoteMemberNotMember,
+	MessagingFailure,
+	NoResponse,
+	UnknownInternalFailure = 0,
+	Succeeded = 1
+};
+
+enum class ERequestPartyInvitationCompletionResult
+{
+	NotLoggedIn = -100,
+	InvitePending,
+	AlreadyInParty,
+	PartyFull,
+	NoPermission,
+	UnknownInternalFailure = 0,
+	Succeeded = 1
+};
+
+enum class ESendPartyInvitationCompletionResult
+{
+	NotLoggedIn = -100,
+	InvitePending,
+	AlreadyInParty,
+	PartyFull,
+	NoPermission,
+	UnknownInternalFailure = 0,
+	Succeeded = 1
+};
+
+enum class EAcceptPartyInvitationCompletionResult
+{
+	NotLoggedIn = -100,
+	InvitePending,
+	AlreadyInParty,
+	PartyFull,
+	NoPermission,
+	UnknownInternalFailure = 0,
+	Succeeded = 1
+};
+
+enum class ERejectPartyInvitationCompletionResult
+{
+	NotLoggedIn = -100,
+	InvitePending,
+	AlreadyInParty,
+	PartyFull,
+	NoPermission,
+	UnknownInternalFailure = 0,
+	Succeeded = 1
+};
+
+enum class EKickMemberCompletionResult
+{
+	UnknownClientFailure = -100,
+	UnknownParty,
+	LocalMemberNotMember,
+	LocalMemberNotLeader,
+	RemoteMemberNotMember,
+	MessagingFailure,
+	NoResponse,
+	UnknownInternalFailure = 0,
+	Succeeded = 1
+};
+
+enum class EPromoteMemberCompletionResult
+{
+	UnknownClientFailure = -100,
+	UnknownServiceFailure,
+	UnknownParty,
+	LocalMemberNotMember,
+	LocalMemberNotLeader,
+	PromotionAlreadyPending,
+	TargetIsSelf,
+	TargetNotMember,
+	MessagingFailure,
+	NoResponse,
+	UnknownInternalFailure = 0,
+	Succeeded = 1
+};
+
+enum class EInvitationResponse
+{
+	UnknownFailure,
+	BadBuild,
+	Rejected,
+	Accepted
+};
+
+
+inline const TCHAR* ToString(const EMemberChangedReason Value)
+{
+	switch (Value)
+	{
+	case EMemberChangedReason::Disconnected:
+	{
+		return TEXT("Disconnected");
+	}
+	case EMemberChangedReason::Rejoined:
+	{
+		return TEXT("Rejoined");
+	}
+	case EMemberChangedReason::Promoted:
+	{
+		return TEXT("Promoted");
+	}
+	}
+	return TEXT("");
+}
+
+/** @return the stringified version of the enum passed in */
+inline const TCHAR* ToString(const EPartyState Value)
+{
+	switch (Value)
+	{
+	case EPartyState::None:
+	{
+		return TEXT("None");
+	}
+	case EPartyState::CreatePending:
+	{
+		return TEXT("Create pending");
+	}
+	case EPartyState::JoinPending:
+	{
+		return TEXT("Join pending");
+	}
+	case EPartyState::LeavePending:
+	{
+		return TEXT("Leave pending");
+	}
+	case EPartyState::Active:
+	{
+		return TEXT("Active");
+	}
+	case EPartyState::Disconnected:
+	{
+		return TEXT("Disconnected");
+	}
+	}
+	return TEXT("");
+}
+
+inline const TCHAR* ToString(const EMemberExitedReason Value)
+{
+	switch (Value)
+	{
+	case EMemberExitedReason::Unknown:
+	{
+		return TEXT("Unknown");
+	}
+	case EMemberExitedReason::Left:
+	{
+		return TEXT("Left");
+	}
+	case EMemberExitedReason::Removed:
+	{
+		return TEXT("Removed");
+	}
+	case EMemberExitedReason::Kicked:
+	{
+		return TEXT("Kicked");
+	}
+	}
+	return TEXT("");
+}
+
+inline const TCHAR* ToString(const ECreatePartyCompletionResult Value)
+{
+	switch (Value)
+	{
+	case ECreatePartyCompletionResult::UnknownClientFailure:
+	{
+		return TEXT("UnknownClientFailure");
+	}
+	case ECreatePartyCompletionResult::AlreadyCreatingParty:
+	{
+		return TEXT("AlreadyCreatingParty");
+	}
+	case ECreatePartyCompletionResult::AlreadyInParty:
+	{
+		return TEXT("AlreadyInParty");
+	}
+	case ECreatePartyCompletionResult::FailedToCreateMucRoom:
+	{
+		return TEXT("FailedToCreateMucRoom");
+	}
+	case ECreatePartyCompletionResult::NoResponse:
+	{
+		return TEXT("NoResponse");
+	}
+	case ECreatePartyCompletionResult::UnknownInternalFailure:
+	{
+		return TEXT("UnknownInternalFailure");
+	}
+	case ECreatePartyCompletionResult::Succeeded:
+	{
+		return TEXT("Succeeded");
+	}
+	}
+	return TEXT("");
+}
+
+inline const TCHAR* ToString(const ESendPartyInvitationCompletionResult Value)
+{
+	switch (Value)
+	{
+	case ESendPartyInvitationCompletionResult::NotLoggedIn:
+	{
+		return TEXT("NotLoggedIn");
+	}
+	case ESendPartyInvitationCompletionResult::InvitePending:
+	{
+		return TEXT("InvitePending");
+	}
+	case ESendPartyInvitationCompletionResult::AlreadyInParty:
+	{
+		return TEXT("AlreadyInParty");
+	}
+	case ESendPartyInvitationCompletionResult::PartyFull:
+	{
+		return TEXT("PartyFull");
+	}
+	case ESendPartyInvitationCompletionResult::NoPermission:
+	{
+		return TEXT("NoPermission");
+	}
+	case ESendPartyInvitationCompletionResult::UnknownInternalFailure:
+	{
+		return TEXT("UnknownInternalFailure");
+	}
+	case ESendPartyInvitationCompletionResult::Succeeded:
+	{
+		return TEXT("Succeeded");
+	}
+	}
+	return TEXT("");
+}
+
+inline const TCHAR* ToString(const EJoinPartyCompletionResult Value)
+{
+	switch (Value)
+	{
+	case EJoinPartyCompletionResult::UnknownClientFailure:
+	{
+		return TEXT("UnknownClientFailure");
+	}
+	case EJoinPartyCompletionResult::JoinInfoInvalid:
+	{
+		return TEXT("JoinInfoInvalid");
+	}
+	case EJoinPartyCompletionResult::AlreadyJoiningParty:
+	{
+		return TEXT("AlreadyJoiningParty");
+	}
+	case EJoinPartyCompletionResult::AlreadyInParty:
+	{
+		return TEXT("AlreadyInParty");
+	}
+	case EJoinPartyCompletionResult::MessagingFailure:
+	{
+		return TEXT("MessagingFailure");
+	}
+	case EJoinPartyCompletionResult::PartyNotInPendingState:
+	{
+		return TEXT("PartyNotInPendingState");
+	}
+	case EJoinPartyCompletionResult::ResponseFromUnexpectedUser:
+	{
+		return TEXT("ResponseFromUnexpectedUser");
+	}
+	case EJoinPartyCompletionResult::NoSpace:
+	{
+		return TEXT("NoSpace");
+	}
+	case EJoinPartyCompletionResult::NotApproved:
+	{
+		return TEXT("NotApproved");
+	}
+	case EJoinPartyCompletionResult::RequesteeNotMember:
+	{
+		return TEXT("RequesteeNotMember");
+	}
+	case EJoinPartyCompletionResult::RequesteeNotLeader:
+	{
+		return TEXT("RequesteeNotLeader");
+	}
+	case EJoinPartyCompletionResult::NoResponse:
+	{
+		return TEXT("NoResponse");
+	}
+	case EJoinPartyCompletionResult::UnknownInternalFailure:
+	{
+		return TEXT("UnknownInternalFailure");
+	}
+	case EJoinPartyCompletionResult::Succeeded:
+	{
+		return TEXT("Succeeded");
+	}
+	}
+	return TEXT("");
+}
+
+inline const TCHAR* ToString(const ELeavePartyCompletionResult Value)
+{
+	switch (Value)
+	{
+	case ELeavePartyCompletionResult::UnknownClientFailure:
+	{
+		return TEXT("UnknownClientFailure");
+	}
+	case ELeavePartyCompletionResult::LeavePending:
+	{
+		return TEXT("LeavePending");
+	}
+	case ELeavePartyCompletionResult::UnknownLocalUser:
+	{
+		return TEXT("UnknownLocalUser");
+	}
+	case ELeavePartyCompletionResult::UnknownParty:
+	{
+		return TEXT("UnknownParty");
+	}
+	case ELeavePartyCompletionResult::MessagingFailure:
+	{
+		return TEXT("MessagingFailure");
+	}
+	case ELeavePartyCompletionResult::NoResponse:
+	{
+		return TEXT("NoResponse");
+	}
+	case ELeavePartyCompletionResult::UnknownTransportFailure:
+	{
+		return TEXT("UnknownTransportFailure");
+	}
+	case ELeavePartyCompletionResult::UnknownInternalFailure:
+	{
+		return TEXT("UnknownInternalFailure");
+	}
+	case ELeavePartyCompletionResult::Succeeded:
+	{
+		return TEXT("Succeeded");
+	}
+	}
+	return TEXT("");
+}
+
+inline const TCHAR* ToString(const EUpdateConfigCompletionResult Value)
+{
+	switch (Value)
+	{
+	case EUpdateConfigCompletionResult::UnknownClientFailure:
+	{
+		return TEXT("UnknownClientFailure");
+	}
+	case EUpdateConfigCompletionResult::UnknownParty:
+	{
+		return TEXT("UnknownParty");
+	}
+	case EUpdateConfigCompletionResult::LocalMemberNotMember:
+	{
+		return TEXT("LocalMemberNotMember");
+	}
+	case EUpdateConfigCompletionResult::LocalMemberNotLeader:
+	{
+		return TEXT("LocalMemberNotLeader");
+	}
+	case EUpdateConfigCompletionResult::RemoteMemberNotMember:
+	{
+		return TEXT("RemoteMemberNotMember");
+	}
+	case EUpdateConfigCompletionResult::MessagingFailure:
+	{
+		return TEXT("MessagingFailure");
+	}
+	case EUpdateConfigCompletionResult::NoResponse:
+	{
+		return TEXT("NoResponse");
+	}
+	case EUpdateConfigCompletionResult::UnknownInternalFailure:
+	{
+		return TEXT("UnknownInternalFailure");
+	}
+	case EUpdateConfigCompletionResult::Succeeded:
+	{
+		return TEXT("Succeeded");
+	}
+	}
+	return TEXT("");
+}
+
+inline const TCHAR* ToString(const EKickMemberCompletionResult Value)
+{
+	switch (Value)
+	{
+	case EKickMemberCompletionResult::UnknownClientFailure:
+	{
+		return TEXT("UnknownClientFailure");
+	}
+	case EKickMemberCompletionResult::UnknownParty:
+	{
+		return TEXT("UnknownParty");
+	}
+	case EKickMemberCompletionResult::LocalMemberNotMember:
+	{
+		return TEXT("LocalMemberNotMember");
+	}
+	case EKickMemberCompletionResult::LocalMemberNotLeader:
+	{
+		return TEXT("LocalMemberNotLeader");
+	}
+	case EKickMemberCompletionResult::RemoteMemberNotMember:
+	{
+		return TEXT("RemoteMemberNotMember");
+	}
+	case EKickMemberCompletionResult::MessagingFailure:
+	{
+		return TEXT("MessagingFailure");
+	}
+	case EKickMemberCompletionResult::NoResponse:
+	{
+		return TEXT("NoResponse");
+	}
+	case EKickMemberCompletionResult::UnknownInternalFailure:
+	{
+		return TEXT("UnknownInternalFailure");
+	}
+	case EKickMemberCompletionResult::Succeeded:
+	{
+		return TEXT("Succeeded");
+	}
+	}
+	return TEXT("");
+}
+
+inline const TCHAR* ToString(const EPromoteMemberCompletionResult Value)
+{
+	switch (Value)
+	{
+	case EPromoteMemberCompletionResult::UnknownClientFailure:
+	{
+		return TEXT("UnknownClientFailure");
+	}
+	case EPromoteMemberCompletionResult::UnknownParty:
+	{
+		return TEXT("UnknownParty");
+	}
+	case EPromoteMemberCompletionResult::LocalMemberNotMember:
+	{
+		return TEXT("LocalMemberNotMember");
+	}
+	case EPromoteMemberCompletionResult::LocalMemberNotLeader:
+	{
+		return TEXT("LocalMemberNotLeader");
+	}
+	case EPromoteMemberCompletionResult::TargetIsSelf:
+	{
+		return TEXT("TargetIsSelf");
+	}
+	case EPromoteMemberCompletionResult::TargetNotMember:
+	{
+		return TEXT("TargetNotMember");
+	}
+	case EPromoteMemberCompletionResult::MessagingFailure:
+	{
+		return TEXT("MessagingFailure");
+	}
+	case EPromoteMemberCompletionResult::NoResponse:
+	{
+		return TEXT("NoResponse");
+	}
+	case EPromoteMemberCompletionResult::UnknownInternalFailure:
+	{
+		return TEXT("UnknownInternalFailure");
+	}
+	case EPromoteMemberCompletionResult::Succeeded:
+	{
+		return TEXT("Succeeded");
+	}
+	}
+	return TEXT("");
+}
+
+inline const TCHAR* ToString(const PartySystemPermissions::EPresencePermissions Value)
+{
+	switch (Value)
+	{
+	case PartySystemPermissions::EPresencePermissions::DoNotPublish:
+	{
+		return TEXT("DoNotPublish");
+	}
+	case PartySystemPermissions::EPresencePermissions::LeaderPublishIdNonePublishKey:
+	{
+		return TEXT("LeaderPublishIdNonePublishKey");
+	}
+	case PartySystemPermissions::EPresencePermissions::LeaderPublishIdLeaderPublishKey:
+	{
+		return TEXT("LeaderPublishIdLeaderPublishKey");
+	}
+	case PartySystemPermissions::EPresencePermissions::FriendPublishIdNonePublishKey:
+	{
+		return TEXT("FriendPublishIdNonePublishKey");
+	}
+	case PartySystemPermissions::EPresencePermissions::FriendPublishIdLeaderPublishKey:
+	{
+		return TEXT("FriendPublishIdLeaderPublishKey");
+	}
+	case PartySystemPermissions::EPresencePermissions::FriendPublishIdFriendPublishKey:
+	{
+		return TEXT("FriendPublishIdFriendPublishKey");
+	}
+	case PartySystemPermissions::EPresencePermissions::AnyonePublishIdNonePublishKey:
+	{
+		return TEXT("AnyonePublishIdNonePublishKey");
+	}
+	case PartySystemPermissions::EPresencePermissions::AnyonePublishIdLeaderPublishKey:
+	{
+		return TEXT("AnyonePublishIdLeaderPublishKey");
+	}
+	case PartySystemPermissions::EPresencePermissions::AnyonePublishIdFriendPublishKey:
+	{
+		return TEXT("AnyonePublishIdFriendPublishKey");
+	}
+	case PartySystemPermissions::EPresencePermissions::AnyonePublishIdAnyonePublishKey:
+	{
+		return TEXT("AnyonePublishIdAnyonePublishKey");
+	}
+	}
+	return TEXT("");
+}
+
+inline const TCHAR* ToString(const PartySystemPermissions::EInvitePermissions Value)
+{
+	switch (Value)
+	{
+	case PartySystemPermissions::EInvitePermissions::Leader:
+	{
+		return TEXT("Leader");
+	}
+	case PartySystemPermissions::EInvitePermissions::Friends:
+	{
+		return TEXT("Friends");
+	}
+	case PartySystemPermissions::EInvitePermissions::Anyone:
+	{
+		return TEXT("Anyone");
+	}
+	}
+	return TEXT("");
+}
+
+inline const TCHAR* ToString(const EJoinRequestAction Value)
+{
+	switch (Value)
+	{
+	case EJoinRequestAction::Manual:
+	{
+		return TEXT("Manual");
+	}
+	case EJoinRequestAction::AutoApprove:
+	{
+		return TEXT("AutoApprove");
+	}
+	case EJoinRequestAction::AutoReject:
+	{
+		return TEXT("AutoReject");
+	}
+	}
+	return TEXT("");
+}
+
+inline const TCHAR* ToString(const EInvitationResponse Value)
+{
+	switch (Value)
+	{
+	case EInvitationResponse::UnknownFailure:
+	{
+		return TEXT("UnknownFailure");
+	}
+	case EInvitationResponse::BadBuild:
+	{
+		return TEXT("BadBuild");
+	}
+	case EInvitationResponse::Rejected:
+	{
+		return TEXT("Rejected");
+	}
+	case EInvitationResponse::Accepted:
+	{
+		return TEXT("Accepted");
+	}
+	}
+	return TEXT("");
+}
+
+inline FString ToDebugString(const FPartyConfiguration& PartyConfiguration)
+{
+	return FString::Printf(TEXT("JoinRequestAction(%s) RemoveOnDisconnect(%d) Publish(%s) Chat(%d) Invite(%s) Accepting(%d) Not Accepting Reason(%d) MaxMembers: %d Nickname: %s Description: %s Password: %s"),
+			ToString(PartyConfiguration.JoinRequestAction),
+			PartyConfiguration.bShouldRemoveOnDisconnection,
+			ToString(PartyConfiguration.PresencePermissions),
+			PartyConfiguration.bChatEnabled,
+			ToString(PartyConfiguration.InvitePermissions),
+			PartyConfiguration.bIsAcceptingMembers,
+			PartyConfiguration.NotAcceptingMembersReason,
+			PartyConfiguration.MaxMembers,
+			*PartyConfiguration.Nickname,
+			*PartyConfiguration.Description,
+			PartyConfiguration.Password.IsEmpty() ? TEXT("not set") : *PartyConfiguration.Password
+		);
+}
+
+inline FString ToDebugString(const IOnlinePartyJoinInfo& JoinInfo)
+{
+	return FString::Printf(TEXT("SourceUserId(%s) SourceDisplayName(%s) PartyId(%s) LeaderUserId(%s) LeaderDisplayName(%s) HasKey(%d) HasPassword(%d) IsAcceptingMembers(%d) NotAcceptingReason(%d)"),
+			*(JoinInfo.GetSourceUserId()->ToDebugString()),
+			*(JoinInfo.GetSourceDisplayName()),
+			*(JoinInfo.GetPartyId()->ToDebugString()),
+			JoinInfo.GetLeaderId()->IsValid() ? *(JoinInfo.GetLeaderId()->ToString()) : TEXT("not set"),
+			!JoinInfo.GetLeaderDisplayName().IsEmpty() ? *(JoinInfo.GetLeaderDisplayName()) : TEXT("not set"),
+			JoinInfo.HasKey() ? 1 : 0,
+			JoinInfo.HasPassword() ? 1 : 0,
+			JoinInfo.IsAcceptingMembers() ? 1 : 0,
+			JoinInfo.GetNotAcceptingReason()
+		);
+}
+
+/**
+* Dump state about the party data for debugging
+*/
+inline FString ToDebugString(const FOnlinePartyData& PartyData)
+{
+	FString Result;
+	
+	int32 TotalBytesPerSec = PartyData.TotalPackets ? (PartyData.TotalBytes / PartyData.TotalPackets) : 0;
+	int32 TotalEffectiveBytesPerSec = PartyData.TotalPackets ? (PartyData.TotalEffectiveBytes / PartyData.TotalPackets) : 0;
+
+	Result += FString::Printf(TEXT("%dB [%d B/pkt], %dB [%d B/pkt], Rev: %d"), 
+		PartyData.TotalBytes, TotalBytesPerSec,
+		PartyData.TotalEffectiveBytes, TotalEffectiveBytesPerSec, 
+		PartyData.RevisionCount);
+
+	for (auto Iterator : PartyData.GetKeyValAttrs())
+	{
+		Result += FString::Printf(TEXT(",[%s=%s]"), *(Iterator.Key), *(Iterator.Value.ToString()));
+	}
+	return Result;
+}
diff --git a/Engine/Plugins/Online/OnlineSubsystem/Source/Public/OnlineError.h b/Engine/Plugins/Online/OnlineSubsystem/Source/Public/OnlineError.h
new file mode 100644
--- /dev/null
+++ b/Engine/Plugins/Online/OnlineSubsystem/Source/Public/OnlineError.h
@@ -0,0 +1,44 @@
+// Copyright 1998-2016 Epic Games, Inc. All Rights Reserved.
+
+#pragma once
+
+namespace EOnlineServerConnectionStatus { enum Type; }
+
+/** Generic Error response for OSS calls */
+struct ONLINESUBSYSTEM_API FOnlineError
+{
+public:
+	FOnlineError();
+	explicit FOnlineError(bool bSucceeded);
+	explicit FOnlineError(const FString& ErrorCode);
+	explicit FOnlineError(const FText& ErrorMessage);
+
+	/** Same as the Ctors but can be called any time */
+	void SetFromErrorCode(const FString& ErrorCode);
+	void SetFromErrorMessage(const FText& ErrorMessage);
+
+	/** Converts the HttpResult into a EOnlineServerConnectionStatus */
+	EOnlineServerConnectionStatus::Type GetConnectionStatusFromHttpResult() const;
+
+	/** Code useful when all you have is raw error info from old APIs */
+	static const FString GenericErrorCode;
+
+	/** Call this if you want to log this out (will pick the best string representation) */
+	const TCHAR* ToLogString() const;
+
+public:
+	/** Did the request succeed fully. If this is true the rest of the struct probably doesn't matter */
+	bool bSucceeded;
+
+	/** The HTTP response code. Will be 0 if a connection error occurred or if HTTP was not used */
+	int32 HttpResult;
+
+	/** The raw unparsed error message from server. Used for pass-through error processing by other systems. */
+	FString ErrorRaw;
+
+	/** Intended to be interpreted by code. */
+	FString ErrorCode;
+
+	/** Suitable for display to end user. Guaranteed to be in the current locale (or empty) */
+	FText ErrorMessage;
+};
diff --git a/Engine/Plugins/Online/OnlineSubsystem/Source/Public/OnlineJsonSerializer.h b/Engine/Plugins/Online/OnlineSubsystem/Source/Public/OnlineJsonSerializer.h
new file mode 100644
--- /dev/null
+++ b/Engine/Plugins/Online/OnlineSubsystem/Source/Public/OnlineJsonSerializer.h
@@ -0,0 +1,21 @@
+// Copyright 1998-2016 Epic Games, Inc. All Rights Reserved.
+
+#pragma once
+#include "Json.h"
+#include "OnlineKeyValuePair.h"
+
+#define BEGIN_ONLINE_JSON_SERIALIZER												BEGIN_JSON_SERIALIZER
+#define END_ONLINE_JSON_SERIALIZER													END_JSON_SERIALIZER
+#define ONLINE_JSON_SERIALIZE(JsonName, JsonValue)									JSON_SERIALIZE(JsonName, JsonValue)
+#define ONLINE_JSON_SERIALIZE_ARRAY(JsonName, JsonArray)							JSON_SERIALIZE_ARRAY(JsonName, JsonArray) 
+#define ONLINE_JSON_SERIALIZE_MAP(JsonName, JsonMap)								JSON_SERIALIZE_MAP(JsonName, JsonMap)
+#define ONLINE_JSON_SERIALIZE_SERIALIZABLE(JsonName, JsonValue)						JSON_SERIALIZE_SERIALIZABLE(JsonName, JsonValue) 
+#define ONLINE_JSON_SERIALIZE_ARRAY_SERIALIZABLE(JsonName, JsonArray, ElementType)	JSON_SERIALIZE_ARRAY_SERIALIZABLE(JsonName, JsonArray, ElementType)
+#define ONLINE_JSON_SERIALIZE_MAP_SERIALIZABLE(JsonName, JsonMap, ElementType)		JSON_SERIALIZE_MAP_SERIALIZABLE(JsonName, JsonMap, ElementType) 
+#define ONLINE_JSON_SERIALIZE_OBJECT_SERIALIZABLE(JsonName, JsonSerializableObject)	JSON_SERIALIZE_OBJECT_SERIALIZABLE(JsonName, JsonSerializableObject) 
+
+typedef FJsonSerializable FOnlineJsonSerializable;
+typedef FJsonSerializerWriter<> FOnlineJsonSerializerWriter;
+
+typedef TMap<FString, FVariantData> FJsonSerializeableKeyValueMapVariant;
+
diff --git a/Engine/Plugins/Online/OnlineSubsystem/Source/Public/OnlineKeyValuePair.h b/Engine/Plugins/Online/OnlineSubsystem/Source/Public/OnlineKeyValuePair.h
new file mode 100644
--- /dev/null
+++ b/Engine/Plugins/Online/OnlineSubsystem/Source/Public/OnlineKeyValuePair.h
@@ -0,0 +1,482 @@
+// Copyright 1998-2016 Epic Games, Inc. All Rights Reserved.
+#pragma once
+
+#include "OnlineSubsystemPackage.h"
+
+namespace EOnlineKeyValuePairDataType
+{
+	enum Type
+	{
+		/** Means the data in the OnlineData value fields should be ignored */
+		Empty,
+		/** 32 bit integer */
+		Int32,
+		/** 32 bit unsigned integer */
+		UInt32,
+		/** 64 bit integer */
+		Int64,
+		/** 64 bit unsigned integer */
+		UInt64,
+		/** Double (8 byte) */
+		Double,
+		/** Unicode string */
+		String,
+		/** Float (4 byte) */
+		Float,
+		/** Binary data */
+		Blob,
+		/** bool data (1 byte) */
+		Bool,
+		MAX
+	};
+
+	/** @return the stringified version of the enum passed in */
+	inline const TCHAR* ToString(EOnlineKeyValuePairDataType::Type EnumVal)
+	{
+		switch (EnumVal)
+		{
+		case Empty:
+			return TEXT("Empty");
+		case Int32:
+			return TEXT("Int32");
+		case UInt32:
+			return TEXT("UInt32");
+		case Int64:
+			return TEXT("Int64");
+		case UInt64:
+			return TEXT("UInt64");
+		case Double:
+			return TEXT("Double");
+		case String:
+			return TEXT("String");
+		case Float:
+			return TEXT("Float");
+		case Blob:
+			return TEXT("Blob");
+		case Bool:
+			return TEXT("Bool");
+		default:
+			return TEXT("");
+		}		
+	}
+}
+
+/**
+ *	Associative container for key value pairs
+ */
+template<class KeyType, class ValueType>
+class FOnlineKeyValuePairs : public TMap<KeyType,ValueType>
+{
+	typedef TMap<KeyType, ValueType> Super;
+
+public:
+
+	FORCEINLINE FOnlineKeyValuePairs() {}
+	FORCEINLINE FOnlineKeyValuePairs(FOnlineKeyValuePairs&& Other) : Super(MoveTemp(Other)) {}
+	FORCEINLINE FOnlineKeyValuePairs(const FOnlineKeyValuePairs&  Other) : Super(Other) {}
+	FORCEINLINE FOnlineKeyValuePairs& operator=(FOnlineKeyValuePairs&& Other) { Super::operator=(MoveTemp(Other)); return *this; }
+	FORCEINLINE FOnlineKeyValuePairs& operator=(const FOnlineKeyValuePairs&  Other) { Super::operator=(Other); return *this; }
+
+};
+
+/**
+ *	Container for storing data of variable type
+ */
+class ONLINESUBSYSTEM_API FVariantData
+{
+
+private:
+
+	/** Current data type */
+	EOnlineKeyValuePairDataType::Type Type;
+	/** Union of all possible types that can be stored */
+	union ValueUnion
+	{
+		bool AsBool;
+		int32 AsInt;
+		uint32 AsUInt;
+		float AsFloat;
+		int64 AsInt64;
+		uint64 AsUInt64;
+		double AsDouble;
+		TCHAR* AsTCHAR;
+		struct 
+		{ 
+			uint8* BlobData;
+			uint32 BlobSize;
+		} AsBlob;
+
+		ValueUnion() { FMemory::Memset( this, 0, sizeof( ValueUnion ) ); }
+	} Value;
+
+public:
+
+	/** Constructor */
+	FVariantData() :
+		Type(EOnlineKeyValuePairDataType::Empty)
+	{
+	}
+
+	/** Constructor starting with an initialized value/type */
+	template<typename ValueType> 
+	FVariantData(const ValueType& InData) :
+		Type(EOnlineKeyValuePairDataType::Empty)
+	{
+		SetValue(InData);
+	}
+
+	/**
+	 * Copy constructor. Copies the other into this object
+	 *
+	 * @param Other the other structure to copy
+	 */
+	FVariantData(const FVariantData& Other);
+
+	/**
+	 * Assignment operator. Copies the other into this object
+	 *
+	 * @param Other the other structure to copy
+	 */
+	FVariantData& operator=(const FVariantData& Other);
+
+	/**
+	 * Cleans up the data to prevent leaks
+	 */
+	~FVariantData()
+	{
+		Empty();
+	}
+
+	/**
+	 * Cleans up the existing data and sets the type to ODT_Empty
+	 */
+	void Empty();
+
+	/**
+	 *	Get the key for this key value pair
+	 */
+	const EOnlineKeyValuePairDataType::Type GetType() const
+	{
+		return Type;
+	}
+
+	/**
+	 * Copies the data and sets the type
+	 *
+	 * @param InData the new data to assign
+	 */
+	void SetValue(const FString& InData);
+
+	/**
+	 * Copies the data and sets the type
+	 *
+	 * @param InData the new data to assign
+	 */
+	void SetValue(const TCHAR* InData);
+
+	/**
+	 * Copies the data and sets the type
+	 *
+	 * @param InData the new data to assign
+	 */
+	void SetValue(int32 InData);
+
+	/**
+	 * Copies the data and sets the type
+	 *
+	 * @param InData the new data to assign
+	 */
+	void SetValue(uint32 InData);
+
+	/**
+	 * Copies the data and sets the type
+	 *
+	 * @param InData the new data to assign
+	 */
+	void SetValue(bool InData);
+
+	/**
+	 * Copies the data and sets the type
+	 *
+	 * @param InData the new data to assign
+	 */
+	void SetValue(double InData);
+
+	/**
+	 * Copies the data and sets the type
+	 *
+	 * @param InData the new data to assign
+	 */
+	void SetValue(float InData);
+
+	/**
+	 * Copies the data and sets the type
+	 *
+	 * @param InData the new data to assign
+	 */
+	void SetValue(const TArray<uint8>& InData);
+
+	/**
+	 * Copies the data and sets the type
+	 *
+	 * @param Size the length of the buffer to copy
+	 * @param InData the new data to assign
+	 */
+	void SetValue(uint32 Size, const uint8* InData);
+
+	/**
+	 * Copies the data and sets the type
+	 *
+	 * @param InData the new data to assign
+	 */
+	void SetValue(int64 InData);
+
+	/**
+	 * Copies the data and sets the type
+	 *
+	 * @param InData the new data to assign
+	 */
+	void SetValue(uint64 InData);
+
+	/**
+	 * Copies the data after verifying the type
+	 *
+	 * @param OutData out value that receives the copied data
+	 */
+	void GetValue(FString& OutData) const;
+
+	/**
+	 * Copies the data after verifying the type
+	 *
+	 * @param OutData out value that receives the copied data
+	 */
+	void GetValue(int32& OutData) const;
+
+	/**
+	 * Copies the data after verifying the type
+	 *
+	 * @param OutData out value that receives the copied data
+	 */
+	void GetValue(uint32& OutData) const;
+
+	/**
+	 * Copies the data after verifying the type
+	 *
+	 * @param OutData out value that receives the copied data
+	 */
+	void GetValue(bool& OutData) const;
+
+	/**
+	  Copies the data after verifying the type
+	 *
+	 * @param OutData out value that receives the copied data
+	 */
+	void GetValue(int64& OutData) const;
+
+	/**
+	  Copies the data after verifying the type
+	 *
+	 * @param OutData out value that receives the copied data
+	 */
+	void GetValue(uint64& OutData) const;
+
+	/**
+	 * Copies the data after verifying the type
+	 *
+	 * @param OutData out value that receives the copied data
+	 */
+	void GetValue(float& OutData) const;
+
+	/**
+	 * Copies the data after verifying the type
+	 *
+	 * @param OutData out value that receives the copied data
+	 */
+	void GetValue(TArray<uint8>& OutData) const;
+
+	/**
+	 * Copies the data after verifying the type.
+	 * NOTE: Performs a deep copy so you are responsible for freeing the data
+	 *
+	 * @param OutSize out value that receives the size of the copied data
+	 * @param OutData out value that receives the copied data
+	 */
+	void GetValue(uint32& OutSize,uint8** OutData) const;
+
+	/**
+	 * Copies the data after verifying the type
+	 *
+	 * @param OutData out value that receives the copied data
+	 */
+	void GetValue(double& OutData) const;
+
+	/**
+	 * Increments the value by the specified amount
+	 * 
+	 * @param IncBy the amount to increment by
+	 */
+	template<typename TYPE, EOnlineKeyValuePairDataType::Type ENUM_TYPE>
+	FORCEINLINE void Increment(TYPE IncBy)
+	{
+		checkSlow(Type == EOnlineKeyValuePairDataType::Int32 || Type == EOnlineKeyValuePairDataType::Int64 ||
+                  Type == EOnlineKeyValuePairDataType::UInt32 || Type == EOnlineKeyValuePairDataType::UInt64 ||
+				  Type == EOnlineKeyValuePairDataType::Float || Type == EOnlineKeyValuePairDataType::Double);
+		if (Type == ENUM_TYPE)
+		{
+			*(TYPE*)&Value += IncBy;
+		}
+	}
+
+	/**
+	 * Decrements the value by the specified amount
+	 *
+	 * @param DecBy the amount to decrement by
+	 */
+	template<typename TYPE, EOnlineKeyValuePairDataType::Type ENUM_TYPE>
+	FORCEINLINE void Decrement(TYPE DecBy)
+	{
+		checkSlow(Type == EOnlineKeyValuePairDataType::Int32 || Type == EOnlineKeyValuePairDataType::Int64 ||
+				  Type == EOnlineKeyValuePairDataType::UInt32 || Type == EOnlineKeyValuePairDataType::UInt64 ||
+				  Type == EOnlineKeyValuePairDataType::Float || Type == EOnlineKeyValuePairDataType::Double);
+		if (Type == ENUM_TYPE)
+		{
+			*(TYPE*)&Value -= DecBy;
+		}
+	}
+
+	/**
+	 * Converts the data into a string representation
+	 */
+	FString ToString() const;
+
+	/**
+	 * Converts the string to the specified type of data for this setting
+	 *
+	 * @param NewValue the string value to convert
+	 *
+	 * @return true if it was converted, false otherwise
+	 */
+	bool FromString(const FString& NewValue);
+
+	/** @return The type as a string */
+	const TCHAR* GetTypeString() const
+	{
+		return EOnlineKeyValuePairDataType::ToString(Type);
+	}
+
+	/**
+	 * Convert variant data to json object with "type,value" fields
+	 *
+	 * @return json object representation
+	 */
+	TSharedRef<class FJsonObject> ToJson() const;
+	
+	/**
+	 * Convert json object to variant data from "type,value" fields
+	 *
+	 * @param JsonObject json to convert from
+	 * @return true if conversion was successful
+	 */
+	bool FromJson(const TSharedRef<class FJsonObject>& JsonObject);
+
+	/**
+	 * Comparison of two settings data classes
+	 *
+	 * @param Other the other settings data to compare against
+	 *
+	 * @return true if they are equal, false otherwise
+	 */
+	bool operator==(const FVariantData& Other) const;
+	bool operator!=(const FVariantData& Other) const;
+};
+
+/**
+ * Helper class for converting from UStruct to FVariantData and back
+ * only very basic flat UStructs with POD types are supported
+ */
+class ONLINESUBSYSTEM_API FVariantDataConverter
+{
+public:
+	/**
+	 * Convert a UStruct into a variant mapping table
+	 *
+	 * @param StructDefinition layout of the UStruct
+	 * @param Struct actual UStruct data
+	 * @param OutVariantMap container for outgoing data
+	 * @param CheckFlags property must have this flag to be serialized
+	 * @param SkipFlags property cannot have this flag to be serialized
+	 * 
+	 * @return true if it was successful, false otherwise
+	 */
+	static bool UStructToVariantMap(const UStruct* StructDefinition, const void* Struct, FOnlineKeyValuePairs<FString, FVariantData>& OutVariantMap, int64 CheckFlags, int64 SkipFlags);
+
+	/**
+	 * Convert a single UProperty to an FVariantData
+	 *
+	 * @param Property definition of the property
+	 * @param Value actual property data
+	 * @param CheckFlags property must have this flag to be serialized
+	 * @param SkipFlags property cannot have this flag to be serialized
+	 * @param OutVariantData container for outgoing data
+	 * 
+	 * @return true if it was successful, false otherwise
+	 */
+	static bool UPropertyToVariantData(UProperty* Property, const void* Value, int64 CheckFlags, int64 SkipFlags, FVariantData& OutVariantData);
+
+public:
+
+	/**
+	 * Convert a map of FVariantData elements to a UStruct
+	 *
+	 * @param VariantMap Input variant data
+	 * @param StructDefinition layout of the UStruct
+	 * @param OutStruct output container for UStruct data
+	 * @param CheckFlags property must have this flag to be serialized
+	 * @param SkipFlags property cannot have this flag to be serialized
+	 * 
+	 * @return true if it was successful, false otherwise
+	 */
+	static bool VariantMapToUStruct(const FOnlineKeyValuePairs<FString, FVariantData>& VariantMap, const UStruct* StructDefinition, void* OutStruct, int64 CheckFlags, int64 SkipFlags);
+
+	/**
+	 * Convert an FVariantData to a UProperty
+	 *
+	 * @param Variant Input variant data
+	 * @param Property definition of the property
+	 * @param OutValue outgoing property data container
+	 * @param CheckFlags property must have this flag to be serialized
+	 * @param SkipFlags property cannot have this flag to be serialized
+	 * 
+	 * @return true if it was successful, false otherwise
+	 */
+	static bool VariantDataToUProperty(const FVariantData* Variant, UProperty* Property, void* OutValue, int64 CheckFlags, int64 SkipFlags);
+	
+private:
+
+	/**
+	 * Convert a single UProperty to an FVariantData
+	 *
+	 * @param Property definition of the property
+	 * @param Value actual property data
+	 * @param OutVariantData container for outgoing data
+	 * @param CheckFlags property must have this flag to be serialized
+	 * @param SkipFlags property cannot have this flag to be serialized
+	 * 
+	 * @return true if it was successful, false otherwise
+	 */
+	static bool ConvertScalarUPropertyToVariant(UProperty* Property, const void* Value, FVariantData& OutVariantData, int64 CheckFlags, int64 SkipFlags);
+
+	/**
+	 * Convert an FVariantData to a UProperty
+	 *
+	 * @param Variant Input variant data
+	 * @param Property definition of the property
+	 * @param OutValue outgoing property data container
+	 * @param CheckFlags property must have this flag to be serialized
+	 * @param SkipFlags property cannot have this flag to be serialized
+	 * 
+	 * @return true if it was successful, false otherwise
+	 */
+	static bool ConvertScalarVariantToUProperty(const FVariantData* Variant, UProperty* Property, void* OutValue, int64 CheckFlags, int64 SkipFlags);
+};
+
diff --git a/Engine/Plugins/Online/OnlineSubsystem/Source/Public/OnlineNotificationHandler.h b/Engine/Plugins/Online/OnlineSubsystem/Source/Public/OnlineNotificationHandler.h
new file mode 100644
--- /dev/null
+++ b/Engine/Plugins/Online/OnlineSubsystem/Source/Public/OnlineNotificationHandler.h
@@ -0,0 +1,94 @@
+// Copyright 1998-2016 Epic Games, Inc. All Rights Reserved.
+
+#pragma once
+#include "OnlineSubsystemPackage.h"
+#include "OnlineSubsystemTypes.h"
+#include "OnlineNotificationTransportInterface.h"
+
+// forward declarations
+class IOnlineNotificationTransport;
+struct FOnlineNotification;
+
+/** Whether a handler function handled a particular notification */
+enum class EOnlineNotificationResult
+{
+	None,		// No handling occurred
+	Handled,	// Notification was handled
+};
+
+/**
+* Delegate type for handling a notification
+*
+* The first parameter is a notification structure
+* Return result code to indicate if notification has been handled
+*/
+DECLARE_DELEGATE_RetVal_OneParam(EOnlineNotificationResult, FHandleOnlineNotificationSignature, const FOnlineNotification&);
+
+
+/** Struct to keep track of bindings */
+struct FOnlineNotificationBinding
+{
+	/** Delegate to call when this binding is activated */
+	FHandleOnlineNotificationSignature NotificationDelegate;
+
+	FOnlineNotificationBinding()
+	{}
+
+	FOnlineNotificationBinding(const FHandleOnlineNotificationSignature& InNotificationDelegate)
+		: NotificationDelegate(InNotificationDelegate)
+	{}
+};
+
+/** This class is a static manager used to track notification transports and map the delivered notifications to subscribed notification handlers */
+class ONLINESUBSYSTEM_API FOnlineNotificationHandler
+{
+
+protected:
+
+	typedef TMap< FString, TArray<FOnlineNotificationBinding> > NotificationTypeBindingsMap;
+
+	/** Map from type of notification to the delegate to call */
+	NotificationTypeBindingsMap SystemBindingMap;
+
+	/** Map from player and type of notification to the delegate to call */
+	TMap< FString, NotificationTypeBindingsMap > PlayerBindingMap;
+
+public:
+
+	/** Lifecycle is managed by OnlineSubSystem, all access should be made through there */
+	FOnlineNotificationHandler()
+	{
+	}
+
+	// SYSTEM NOTIFICATION BINDINGS
+
+	/** Add a notification binding for a type */
+	FDelegateHandle AddSystemNotificationBinding_Handle(FString NotificationType, const FOnlineNotificationBinding& NewBinding);
+
+	/** Remove the notification handler for a type */
+	void RemoveSystemNotificationBinding(FString NotificationType, FDelegateHandle RemoveHandle);
+
+	/** Resets all system notification handlers */
+	void ResetSystemNotificationBindings();
+
+	// PLAYER NOTIFICATION BINDINGS
+
+	/** Add a notification binding for a type */
+	FDelegateHandle AddPlayerNotificationBinding_Handle(const FUniqueNetId& PlayerId, FString NotificationType, const FOnlineNotificationBinding& NewBinding);
+
+	/** Remove the player notification handler for a type */
+	void RemovePlayerNotificationBinding(const FUniqueNetId& PlayerId, FString NotificationType, FDelegateHandle RemoveHandle);
+
+	/** Resets a player's notification handlers */
+	void ResetPlayerNotificationBindings(const FUniqueNetId& PlayerId);
+
+	/** Resets all player notification handlers */
+	void ResetAllPlayerNotificationBindings();
+
+	// RECEIVING NOTIFICATIONS
+
+	/** Deliver a notification to the appropriate handler for that player/msg type.  Called by NotificationTransport implementations. */
+	void DeliverNotification(const FOnlineNotification& Notification);
+};
+
+typedef TSharedPtr<FOnlineNotificationHandler, ESPMode::ThreadSafe> FOnlineNotificationHandlerPtr;
