Copying //Tasks/UE4/Dev-Online-Plugins-v3 to //UE4/Dev-Main (Source: //Tasks/UE4/Dev-Online-Plugins-v3 @ 3047294)

diff --git a/Engine/Plugins/Online/OnlineSubsystem/Source/Public/OnlineSessionSettings.h b/Engine/Plugins/Online/OnlineSubsystem/Source/Public/OnlineSessionSettings.h
new file mode 100644
--- /dev/null
+++ b/Engine/Plugins/Online/OnlineSubsystem/Source/Public/OnlineSessionSettings.h
@@ -0,0 +1,587 @@
+// Copyright 1998-2016 Epic Games, Inc. All Rights Reserved.
+
+#pragma once
+
+#include "OnlineSubsystemTypes.h"
+#include "OnlineKeyValuePair.h"
+#include "OnlineSubsystemPackage.h"
+#include "OnlineSubsystemSessionSettings.h"
+
+/** Setting describing the name of the current map (value is FString) */
+#define SETTING_MAPNAME FName(TEXT("MAPNAME"))
+/** Setting describing the number of bots in the session (value is int32) */
+#define SETTING_NUMBOTS FName(TEXT("NUMBOTS"))
+/** Setting describing the game mode of the session (value is FString) */
+#define SETTING_GAMEMODE FName(TEXT("GAMEMODE"))
+/** Setting describing the beacon host port (value is int32) */
+#define SETTING_BEACONPORT FName(TEXT("BEACONPORT"))
+/** Server responds to Qos beacon requests (value is int32) */
+#define SETTING_QOS FName(TEXT("QOS"))
+/** Setting describing the region of the world you are in (value is FString) */
+#define SETTING_REGION FName(TEXT("REGION"))
+/** Setting describing the unique id of a datacenter (value is FString) */
+#define SETTING_DCID FName(TEXT("DCID"))
+/** Number of players needed to fill out this session (value is int32) */
+#define SETTING_NEEDS FName(TEXT("NEEDS"))
+/** Second key for "needs" because can't set same value with two criteria (value is int32) */
+#define SETTING_NEEDSSORT FName(TEXT("NEEDSSORT"))
+
+/** 8 user defined integer params to be used when filtering searches for sessions */
+#define SETTING_CUSTOMSEARCHINT1 FName(TEXT("CUSTOMSEARCHINT1"))
+#define SETTING_CUSTOMSEARCHINT2 FName(TEXT("CUSTOMSEARCHINT2"))
+#define SETTING_CUSTOMSEARCHINT3 FName(TEXT("CUSTOMSEARCHINT3"))
+#define SETTING_CUSTOMSEARCHINT4 FName(TEXT("CUSTOMSEARCHINT4"))
+#define SETTING_CUSTOMSEARCHINT5 FName(TEXT("CUSTOMSEARCHINT5"))
+#define SETTING_CUSTOMSEARCHINT6 FName(TEXT("CUSTOMSEARCHINT6"))
+#define SETTING_CUSTOMSEARCHINT7 FName(TEXT("CUSTOMSEARCHINT7"))
+#define SETTING_CUSTOMSEARCHINT8 FName(TEXT("CUSTOMSEARCHINT8"))
+
+/** TODO ONLINE Settings to consider */
+/** The server's nonce for this session */
+/** Whether the game is an invitation or searched for game */
+/** The ping of the server in milliseconds (-1 means the server was unreachable) */
+/** Whether this server is a dedicated server or not */
+/** Represents how good a match this is in a range from 0 to 1 */
+/** Whether there is a skill update in progress or not (don't do multiple at once) */
+/** Whether to shrink the session slots when a player leaves the match or not */
+
+/**
+ *	One setting describing an online session
+ * contains a key, value and how this setting is advertised to others, if at all
+ */
+struct FOnlineSessionSetting
+{
+public:
+
+	/** Settings value */
+	FVariantData Data;
+	/** How is this session setting advertised with the backend or searches */
+	EOnlineDataAdvertisementType::Type AdvertisementType;
+
+	/** Default constructor, used when serializing a network packet */
+	FOnlineSessionSetting() :
+		AdvertisementType(EOnlineDataAdvertisementType::DontAdvertise)
+	{
+	}
+
+	/** Constructor for settings created/defined on the host for a session */
+	template<typename Type> 
+	FOnlineSessionSetting(const Type& InData) :
+		Data(InData),
+		AdvertisementType(EOnlineDataAdvertisementType::DontAdvertise)
+	{
+	}
+
+	/** Constructor for settings created/defined on the host for a session */
+	template<typename Type> 
+	FOnlineSessionSetting(const Type& InData, EOnlineDataAdvertisementType::Type InAdvertisementType) :
+		Data(InData),
+		AdvertisementType(InAdvertisementType)
+	{
+	}
+
+	/**
+	 *	Comparison operator
+	 */
+	bool operator==(const FOnlineSessionSetting& Other) const
+	{
+		// Advertisement type not compared because it is not passed to clients
+		return Data == Other.Data;
+	}
+
+	FString ToString() const
+	{
+		return FString::Printf(TEXT("%s :%s"), *Data.ToString(), EOnlineDataAdvertisementType::ToString(AdvertisementType));
+	}
+};
+
+/** Type defining an array of session settings accessible by key */
+typedef FOnlineKeyValuePairs<FName, FOnlineSessionSetting> FSessionSettings;
+
+/**
+ *	One search parameter in an online session query
+ *  contains a value and how this setting is compared
+ */
+class FOnlineSessionSearchParam 
+{
+private:
+	/** Hidden on purpose */
+	FOnlineSessionSearchParam() :
+	   ComparisonOp(EOnlineComparisonOp::Equals)
+	{
+	}
+
+public:
+
+	/** Search value */
+	FVariantData Data;
+	/** How is this session setting compared on the backend searches */
+	EOnlineComparisonOp::Type ComparisonOp;
+
+	/** Constructor for setting search parameters in a query */
+	template<typename Type> 
+	FOnlineSessionSearchParam(const Type& InData) :
+		Data(InData),
+		ComparisonOp(EOnlineComparisonOp::Equals)
+	{
+	}
+
+	/** Constructor for setting search parameters in a query */
+	template<typename Type> 
+	FOnlineSessionSearchParam(const Type& InData, EOnlineComparisonOp::Type InComparisonOp) :
+		Data(InData),
+		ComparisonOp(InComparisonOp)
+	{
+	}
+
+	/**
+	 *	Comparison operator
+	 */
+	bool operator==(const FOnlineSessionSearchParam& Other) const
+	{
+		// Don't compare ComparisonOp so we don't get the same data with different ops
+		return Data == Other.Data;
+	}
+
+	FString ToString() const
+	{
+		return FString::Printf(TEXT("Value=%s : %s"), *Data.ToString(), EOnlineComparisonOp::ToString(ComparisonOp));
+	}
+};
+
+/** Type defining an array of search parameters accessible by key */
+typedef FOnlineKeyValuePairs<FName, FOnlineSessionSearchParam> FSearchParams;
+
+/**
+ *	Container for all parameters describing a single session search
+ */
+class ONLINESUBSYSTEM_API FOnlineSearchSettings
+{
+public:
+
+	/** Array of custom search settings */
+	FSearchParams SearchParams;
+
+	FOnlineSearchSettings()
+	{
+	}
+
+	virtual ~FOnlineSearchSettings() 
+	{
+	}
+
+	/**
+	 *	Sets a key value pair combination that defines a search parameter
+	 *
+	 * @param Key key for the setting
+	 * @param Value value of the setting
+	 * @param InType type of comparison
+	 */
+	template<typename ValueType> 
+	void Set(FName Key, const ValueType& Value, EOnlineComparisonOp::Type InType);
+
+	/**
+	 *	Gets a key value pair combination that defines a search parameter
+	 *
+	 * @param Key key for the setting
+	 * @param Value value of the setting
+	 *
+	 * @return true if found, false otherwise
+	 */
+	template<typename ValueType> 
+	bool Get(FName Key, ValueType& Value) const;
+
+	/** 
+	 * Retrieve a search parameter comparison op
+	 *
+	 * @param Key key of the setting
+	 * 
+	 * @return the comparison op for the setting
+	 */
+	EOnlineComparisonOp::Type GetComparisonOp(FName Key) const;
+};
+
+/**
+ *	Container for all settings describing a single online session
+ */
+class ONLINESUBSYSTEM_API FOnlineSessionSettings
+{
+
+public:
+
+	/** The number of publicly available connections advertised */
+	int32 NumPublicConnections;
+	/** The number of connections that are private (invite/password) only */
+	int32 NumPrivateConnections;
+	/** Whether this match is publicly advertised on the online service */
+	bool bShouldAdvertise;
+	/** Whether joining in progress is allowed or not */
+	bool bAllowJoinInProgress;
+	/** This game will be lan only and not be visible to external players */
+	bool bIsLANMatch;
+	/** Whether the server is dedicated or player hosted */
+	bool bIsDedicated;
+	/** Whether the match should gather stats or not */
+	bool bUsesStats;
+	/** Whether the match allows invitations for this session or not */
+	bool bAllowInvites;
+	/** Whether to display user presence information or not */
+	bool bUsesPresence;
+	/** Whether joining via player presence is allowed or not */
+	bool bAllowJoinViaPresence;
+	/** Whether joining via player presence is allowed for friends only or not */
+	bool bAllowJoinViaPresenceFriendsOnly;
+	/** Whether the server employs anti-cheat (punkbuster, vac, etc) */
+	bool bAntiCheatProtected;
+	/** Used to keep different builds from seeing each other during searches */
+	int32 BuildUniqueId;
+	/** Array of custom session settings */
+	FSessionSettings Settings;
+
+	/** Default constructor, used when serializing a network packet */
+	FOnlineSessionSettings() :
+		NumPublicConnections(0),
+		NumPrivateConnections(0),
+		bShouldAdvertise(false),
+		bAllowJoinInProgress(false),
+		bIsLANMatch(false),
+		bIsDedicated(false),
+		bUsesStats(false),
+		bAllowInvites(false),
+		bUsesPresence(false),
+		bAllowJoinViaPresence(false),
+		bAllowJoinViaPresenceFriendsOnly(false),
+		bAntiCheatProtected(false),
+		BuildUniqueId(0)
+	{
+// 		Set(SETTING_MAPNAME, FString(TEXT("")), EOnlineDataAdvertisementType::ViaOnlineService);
+// 		Set(SETTING_NUMBOTS, 0, EOnlineDataAdvertisementType::ViaOnlineService);
+// 		Set(SETTING_GAMEMODE, FString(TEXT("")), EOnlineDataAdvertisementType::ViaOnlineService);
+	}
+
+	virtual ~FOnlineSessionSettings() {}
+
+	/**
+	 *	Sets a key value pair combination that defines a session setting
+	 *
+	 * @param Key key for the setting
+	 * @param Value value of the setting
+	 * @param InType type of online advertisement
+	 */
+	template<typename ValueType> 
+	void Set(FName Key, const ValueType& Value, EOnlineDataAdvertisementType::Type InType);
+
+	/**
+	 *	Sets a key value pair combination that defines a session setting
+	 * from an existing session setting
+	 *
+	 * @param Key key for the setting
+	 * @param SrcSetting setting values
+	 */
+	void Set(FName Key, const FOnlineSessionSetting& SrcSetting);
+
+	/**
+	 *	Gets a key value pair combination that defines a session setting
+	 *
+	 * @param Key key for the setting
+	 * @param Value value of the setting
+	 *
+	 * @return true if found, false otherwise
+	 */
+	template<typename ValueType> 
+	bool Get(FName Key, ValueType& Value) const;
+
+	/**
+	 *  Removes a key value pair combination
+	 *
+	 * @param Key key to remove
+	 * 
+	 * @return true if found and removed, false otherwise
+	 */
+	bool Remove(FName Key);
+
+	/** 
+	 * Retrieve a session setting's advertisement type
+	 *
+	 * @param Key key of the setting
+	 * 
+	 * @return the advertisement type for the setting
+	 */
+	EOnlineDataAdvertisementType::Type GetAdvertisementType(FName Key) const;
+
+};
+
+/** Basic session information serializable into a NamedSession or SearchResults */
+class FOnlineSession
+{
+
+public:
+
+	/** Owner of the session */
+	TSharedPtr<const FUniqueNetId> OwningUserId;
+	/** Owner name of the session */
+	FString OwningUserName;
+	/** The settings associated with this session */
+	FOnlineSessionSettings SessionSettings;
+	/** The platform specific session information */
+	TSharedPtr<class FOnlineSessionInfo> SessionInfo;
+	/** The number of private connections that are available (read only) */
+	int32 NumOpenPrivateConnections;
+	/** The number of publicly available connections that are available (read only) */
+	int32 NumOpenPublicConnections;
+
+	/** Default constructor, used when serializing a network packet */
+	FOnlineSession() :
+		OwningUserId(NULL),
+		SessionInfo(NULL),
+		NumOpenPrivateConnections(0),
+		NumOpenPublicConnections(0)
+	{
+	}
+
+	/** Constructor */
+	FOnlineSession(const FOnlineSessionSettings& InSessionSettings) :
+		OwningUserId(NULL),
+		SessionSettings(InSessionSettings),
+		SessionInfo(NULL),
+		NumOpenPrivateConnections(0),
+		NumOpenPublicConnections(0)
+	{
+	}
+
+	virtual ~FOnlineSession() {}
+};
+
+/** Holds the per session information for named sessions */
+class FNamedOnlineSession : public FOnlineSession
+{
+protected:
+	FNamedOnlineSession() :
+		SessionName(NAME_None),
+		HostingPlayerNum(INDEX_NONE),
+		SessionState(EOnlineSessionState::NoSession)
+	{
+	}
+
+public:
+
+	/** The name of the session */
+	const FName SessionName;
+	/** Index of the player who created the session [host] or joined it [client] */
+	int32 HostingPlayerNum;	
+	/** Whether or not the local player is hosting this session */
+	bool bHosting;
+
+	/** NetId of the local player that created this named session.  Could be the host, or a player joining a session. Will entirely replace HostingPlayerNum */
+	TSharedPtr<const FUniqueNetId> LocalOwnerId;
+
+	/** List of players registered in the session */
+	TArray< TSharedRef<const FUniqueNetId> > RegisteredPlayers;
+	/** State of the session (game thread write only) */
+	EOnlineSessionState::Type SessionState;
+
+	/** Constructor used to create a named session directly */
+	FNamedOnlineSession(FName InSessionName, const FOnlineSessionSettings& InSessionSettings) :
+		FOnlineSession(InSessionSettings),
+		SessionName(InSessionName),
+		HostingPlayerNum(INDEX_NONE),
+		bHosting(false),
+		SessionState(EOnlineSessionState::NoSession)
+	{
+	}
+
+	/** Constructor used to create a named session directly */
+	FNamedOnlineSession(FName InSessionName, const FOnlineSession& Session) :
+		FOnlineSession(Session),
+		SessionName(InSessionName),
+		HostingPlayerNum(INDEX_NONE),
+		bHosting(false),
+		SessionState(EOnlineSessionState::NoSession)
+	{
+	}
+
+	virtual ~FNamedOnlineSession() {}
+
+	/**
+	 * Calculate the possible joinability state of this session
+	 * check the values from left to right in order of precedence
+	 *
+	 * @param bPublicJoinable [out] is the game joinable by anyone at all
+	 * @param bFriendJoinable [out] is the game joinable by friends via presence (doesn't require invite)
+	 * @param bInviteOnly [out] is the game joinable via explicit invites
+	 * @param bAllowInvites [out] are invites possible (use with bInviteOnly to determine if invites are available right now)
+	 *
+	 * @return true if the out params are valid, false otherwise
+	 */
+	bool GetJoinability(bool& bPublicJoinable, bool& bFriendJoinable, bool& bInviteOnly, bool& bAllowInvites) const
+	{
+		// Only states that have a valid session are considered
+		if (SessionState != EOnlineSessionState::NoSession && SessionState != EOnlineSessionState::Creating && SessionState != EOnlineSessionState::Destroying)
+		{
+			bool bAllowJIP = SessionSettings.bAllowJoinInProgress || (SessionState != EOnlineSessionState::Starting && SessionState != EOnlineSessionState::InProgress);
+			if (bAllowJIP)
+			{
+				bPublicJoinable = SessionSettings.bShouldAdvertise || SessionSettings.bAllowJoinViaPresence;
+				bFriendJoinable = SessionSettings.bAllowJoinViaPresenceFriendsOnly;
+				bInviteOnly = !bPublicJoinable && !bFriendJoinable && SessionSettings.bAllowInvites;
+				bAllowInvites = SessionSettings.bAllowInvites;
+			}
+			else
+			{
+				bPublicJoinable = false;
+				bFriendJoinable = false;
+				bInviteOnly = false;
+				bAllowInvites = false;
+			}
+
+			// Valid session, joinable or otherwise
+			return true;
+		}
+		
+		// Invalid session
+		return false;
+	}
+};
+
+/** Value returned on unreachable or otherwise bad search results */
+#define MAX_QUERY_PING 9999
+
+/** Representation of a single search result from a FindSession() call */
+class FOnlineSessionSearchResult
+{
+public:
+	/** All advertised session information */
+	FOnlineSession Session;
+	/** Ping to the search result, MAX_QUERY_PING is unreachable */
+	int32 PingInMs;
+
+	FOnlineSessionSearchResult() : 
+		PingInMs(MAX_QUERY_PING)
+	{}
+
+	~FOnlineSessionSearchResult() {}
+
+	/** Copy Constructor */
+	FOnlineSessionSearchResult(const FOnlineSessionSearchResult& Src) :
+		Session(Src.Session),
+		PingInMs(Src.PingInMs)
+	{
+	}
+
+	/**
+	 *	@return true if the search result is valid, false otherwise
+	 */
+	bool IsValid() const
+	{
+		return (Session.OwningUserId.IsValid() && Session.SessionInfo.IsValid() && Session.SessionInfo->IsValid());
+	}
+};
+
+/** Search only for dedicated servers (value is true/false) */
+#define SEARCH_DEDICATED_ONLY FName(TEXT("DEDICATEDONLY"))
+/** Search for empty servers only (value is true/false) */
+#define SEARCH_EMPTY_SERVERS_ONLY FName(TEXT("EMPTYONLY"))
+/** Search for non empty servers only (value is true/false) */
+#define SEARCH_NONEMPTY_SERVERS_ONLY FName(TEXT("NONEMPTYONLY"))
+/** Search for secure servers only (value is true/false) */
+#define SEARCH_SECURE_SERVERS_ONLY FName(TEXT("SECUREONLY"))
+/** Search for presence sessions only (value is true/false) */
+#define SEARCH_PRESENCE FName(TEXT("PRESENCESEARCH"))
+/** Search for a match with min player availability (value is int) */
+#define SEARCH_MINSLOTSAVAILABLE FName(TEXT("MINSLOTSAVAILABLE"))
+/** Exclude all matches where any unique ids in a given array are present (value is string of the form "uniqueid1;uniqueid2;uniqueid3") */
+#define SEARCH_EXCLUDE_UNIQUEIDS FName(TEXT("EXCLUDEUNIQUEIDS"))
+/** User ID to search for session of */
+#define SEARCH_USER FName(TEXT("SEARCHUSER"))
+/** Keywords to match in session search */
+#define SEARCH_KEYWORDS FName(TEXT("SEARCHKEYWORDS"))
+/** If set, use the named Xbox Live hopper to find a session via matchmaking (value is a string) */
+#define SEARCH_XBOX_LIVE_HOPPER_NAME FName(TEXT("LIVEHOPPERNAME"))
+/** Which session template from the service configuration to use */
+#define SEARCH_XBOX_LIVE_SESSION_TEMPLATE_NAME FName(TEXT("LIVESESSIONTEMPLATE"))
+
+/** 
+ * Encapsulation of a search for sessions request.
+ * Contains all the search parameters and any search results returned after
+ * the OnFindSessionsCompleteDelegate has triggered
+ * Check the SearchState for Done/Failed state before using the data
+ */
+class FOnlineSessionSearch
+{
+public:
+
+	/** Array of all sessions found when searching for the given criteria */
+	TArray<FOnlineSessionSearchResult> SearchResults;
+	/** State of the search */
+	EOnlineAsyncTaskState::Type SearchState;
+	/** Max number of queries returned by the matchmaking service */
+	int32 MaxSearchResults;
+	/** The query to use for finding matching servers */
+	FOnlineSearchSettings QuerySettings;
+	/** Whether the query is intended for LAN matches or not */
+	bool bIsLanQuery;
+	/**
+	 * Used to sort games into buckets since a the difference in terms of feel for ping
+	 * in the same bucket is often not a useful comparison and skill is better
+	 */
+	int32 PingBucketSize;
+
+	/** Amount of time to wait for the search results. May not apply to all platforms. */
+	float TimeoutInSeconds;
+
+	/** Constructor */
+	FOnlineSessionSearch() :
+		SearchState(EOnlineAsyncTaskState::NotStarted),
+		MaxSearchResults(1),
+		bIsLanQuery(false),
+		PingBucketSize(0),
+		TimeoutInSeconds(0.0f)
+	{
+		QuerySettings.Set(SETTING_MAPNAME, FString(TEXT("")), EOnlineComparisonOp::Equals);
+		QuerySettings.Set(SEARCH_DEDICATED_ONLY, false, EOnlineComparisonOp::Equals);
+		QuerySettings.Set(SEARCH_EMPTY_SERVERS_ONLY, false, EOnlineComparisonOp::Equals);
+		QuerySettings.Set(SEARCH_SECURE_SERVERS_ONLY, false, EOnlineComparisonOp::Equals);
+	}
+
+	virtual ~FOnlineSessionSearch() {}
+
+	/**
+	 *	Give the game a chance to sort the returned results
+	 */
+	virtual void SortSearchResults() {}
+
+	/**
+	 * Get the default session settings for this search type
+	 * Allows games to set reasonable defaults that aren't advertised
+	 * but would be setup for each instantiated search result
+	 */
+	virtual TSharedPtr<FOnlineSessionSettings> GetDefaultSessionSettings() const 
+	{ 
+		return MakeShareable(new FOnlineSessionSettings()); 
+	}
+};
+
+/**
+ * Logs session properties used from the session settings
+ *
+ * @param NamedSession the session to log the information for
+ */
+void ONLINESUBSYSTEM_API DumpNamedSession(const FNamedOnlineSession* NamedSession);
+
+/**
+ * Logs session properties used from the session settings
+ *
+ * @param Session the session to log the information for
+ */
+void ONLINESUBSYSTEM_API DumpSession(const FOnlineSession* Session);
+
+/**
+ * Logs session properties used from the session settings
+ *
+ * @param SessionSettings the session to log the information for
+ */
+void ONLINESUBSYSTEM_API DumpSessionSettings(const FOnlineSessionSettings* SessionSettings);
+
+
+
+
diff --git a/Engine/Plugins/Online/OnlineSubsystem/Source/Public/OnlineSubsystemImpl.h b/Engine/Plugins/Online/OnlineSubsystem/Source/Public/OnlineSubsystemImpl.h
new file mode 100644
--- /dev/null
+++ b/Engine/Plugins/Online/OnlineSubsystem/Source/Public/OnlineSubsystemImpl.h
@@ -0,0 +1,107 @@
+// Copyright 1998-2016 Epic Games, Inc. All Rights Reserved.
+
+#pragma once
+
+#include "OnlineSubsystem.h"
+#include "OnlineSubsystemPackage.h"
+
+DECLARE_DELEGATE(FNextTickDelegate);
+
+/**
+ *	FOnlineSubsystemImpl - common functionality to share across online platforms, not intended for direct use
+ */
+class ONLINESUBSYSTEM_API FOnlineSubsystemImpl 
+	: public IOnlineSubsystem
+{
+private:
+
+	/**
+	 * Exec function handling for Exec() call
+	 */
+	bool HandleFriendExecCommands(UWorld* InWorld, const TCHAR* Cmd, FOutputDevice& Ar);
+	bool HandleSessionExecCommands(UWorld* InWorld, const TCHAR* Cmd, FOutputDevice& Ar);
+
+protected:
+
+	/** Hidden on purpose */
+	FOnlineSubsystemImpl();
+	FOnlineSubsystemImpl(FName InInstanceName);
+
+	/** Instance name (disambiguates PIE instances for example) */
+	FName InstanceName;
+
+	/** Whether or not the online subsystem is in forced dedicated server mode */
+	bool bForceDedicated;
+
+	/** Holds all currently named interfaces */
+	mutable class UNamedInterfaces* NamedInterfaces;
+
+	/** Load in any named interfaces specified by the ini configuration */
+	void InitNamedInterfaces();
+
+	/** Delegate fired when named interfaces are cleaned up at exit */
+	void OnNamedInterfaceCleanup();
+
+	/** Queue to hold callbacks scheduled for next tick using ExecuteNextTick */
+	TQueue<FNextTickDelegate, EQueueMode::Mpsc> NextTickQueue;
+
+	/** Buffer to hold callbacks for the current tick (so it's safe to call ExecuteNextTick within a tick callback) */
+	TArray<FNextTickDelegate> CurrentTickBuffer;
+
+	/** Start Ticker */
+	void StartTicker();
+
+	/** Stop Ticker */
+	void StopTicker();
+	
+	/** Delegate for callbacks to Tick */
+	FDelegateHandle TickHandle;
+
+public:
+	
+	virtual ~FOnlineSubsystemImpl();
+
+	// IOnlineSubsystem
+	virtual bool Shutdown() override;
+	virtual bool IsServer() const override;
+	virtual bool IsDedicated() const override{ return bForceDedicated || IsRunningDedicatedServer(); }
+	virtual void SetForceDedicated(bool bForce) override { bForceDedicated = bForce; }
+	virtual class UObject* GetNamedInterface(FName InterfaceName) override;
+	virtual void SetNamedInterface(FName InterfaceName, class UObject* NewInterface) override;
+	virtual bool IsLocalPlayer(const FUniqueNetId& UniqueId) const override;
+	virtual void SetUsingMultiplayerFeatures(const FUniqueNetId& UniqueId, bool bUsingMP) override {};
+	virtual EOnlineEnvironment::Type GetOnlineEnvironment() const override { return EOnlineEnvironment::Unknown; }
+	virtual IMessageSanitizerPtr GetMessageSanitizer(int32 LocalUserNum, FString& OutAuthTypeToExclude) const override;
+	virtual bool Exec(UWorld* InWorld, const TCHAR* Cmd, FOutputDevice& Ar) override;
+
+	/**
+	 * Tick function
+	 *
+	 * @param DeltaTime	time passed since the last call.
+	 * @return true if should continue ticking
+	 */
+	virtual bool Tick(float DeltaTime);
+
+	/**
+	 * @return the name of the online subsystem instance
+	 */
+	virtual FName GetInstanceName() const override { return InstanceName; }
+
+	/**
+	 * Queue a delegate to be executed on the next tick
+	 */
+	void ExecuteDelegateNextTick(const FNextTickDelegate& Callback);
+
+	/**
+	 * Templated helper for calling ExecuteDelegateNextTick with a lambda function
+	 */
+	template<typename LAMBDA_TYPE>
+	FORCEINLINE void ExecuteNextTick(LAMBDA_TYPE&& Callback)
+	{
+		ExecuteDelegateNextTick(FNextTickDelegate::CreateLambda(Callback));
+	}
+
+	/** Name given to default OSS instances (disambiguates for PIE) */
+	static const FName DefaultInstanceName;
+};
+
diff --git a/Engine/Plugins/Online/OnlineSubsystem/Source/Public/OnlineSubsystemNames.h b/Engine/Plugins/Online/OnlineSubsystem/Source/Public/OnlineSubsystemNames.h
new file mode 100644
--- /dev/null
+++ b/Engine/Plugins/Online/OnlineSubsystem/Source/Public/OnlineSubsystemNames.h
@@ -0,0 +1,45 @@
+// Copyright 1998-2016 Epic Games, Inc. All Rights Reserved.
+
+#ifndef NULL_SUBSYSTEM
+#define NULL_SUBSYSTEM FName(TEXT("NULL"))
+#endif
+
+#ifndef GOOGLEPLAY_SUBSYSTEM
+#define GOOGLEPLAY_SUBSYSTEM FName(TEXT("GOOGLEPLAY"))
+#endif
+
+#ifndef IOS_SUBSYSTEM
+#define IOS_SUBSYSTEM FName(TEXT("IOS"))
+#endif
+
+#ifndef AMAZON_SUBSYSTEM
+#define AMAZON_SUBSYSTEM FName(TEXT("AMAZON"))
+#endif
+
+#ifndef FACEBOOK_SUBSYSTEM
+#define FACEBOOK_SUBSYSTEM FName(TEXT("FACEBOOK"))
+#endif
+
+#ifndef STEAM_SUBSYSTEM
+#define STEAM_SUBSYSTEM FName(TEXT("STEAM"))
+#endif
+
+#ifndef PS4_SUBSYSTEM
+#define PS4_SUBSYSTEM FName(TEXT("PS4"))
+#endif
+
+#ifndef LIVE_SUBSYSTEM
+#define LIVE_SUBSYSTEM FName(TEXT("LIVE"))
+#endif
+
+#ifndef THUNDERHEAD_SUBSYSTEM
+#define THUNDERHEAD_SUBSYSTEM FName(TEXT("THUNDERHEAD"))
+#endif
+
+#ifndef MCP_SUBSYSTEM
+#define MCP_SUBSYSTEM FName(TEXT("MCP"))
+#endif
+
+#ifndef WECHAT_SUBSYSTEM
+#define WECHAT_SUBSYSTEM FName(TEXT("WeChat"))
+#endif
\ No newline at end of file
diff --git a/Engine/Plugins/Online/OnlineSubsystem/Source/Public/OnlineSubsystemTypes.h b/Engine/Plugins/Online/OnlineSubsystem/Source/Public/OnlineSubsystemTypes.h
new file mode 100644
--- /dev/null
+++ b/Engine/Plugins/Online/OnlineSubsystem/Source/Public/OnlineSubsystemTypes.h
@@ -0,0 +1,1399 @@
+// Copyright 1998-2016 Epic Games, Inc. All Rights Reserved.
+
+#pragma once
+#include "UObject/CoreOnline.h"
+#include "OnlineSubsystemPackage.h"
+
+/** Maximum players supported on a given platform */
+#if PLATFORM_XBOXONE
+#define MAX_LOCAL_PLAYERS 4
+#elif PLATFORM_PS4
+#define MAX_LOCAL_PLAYERS 4
+#else
+#define MAX_LOCAL_PLAYERS 1
+#endif
+
+#ifndef ERROR_SUCCESS
+#define ERROR_SUCCESS 0
+#endif
+
+#ifndef E_FAIL
+#define E_FAIL (uint32)-1
+#endif
+
+#ifndef E_NOTIMPL
+#define E_NOTIMPL (uint32)-2
+#endif
+
+#ifndef ERROR_IO_PENDING
+#define ERROR_IO_PENDING 997
+#endif
+
+#ifndef S_OK
+#define S_OK 0
+#endif
+
+/**
+ * Generates a random nonce (number used once) of the desired length
+ *
+ * @param Nonce the buffer that will get the randomized data
+ * @param Length the number of bytes to generate random values for
+ */
+inline void GenerateNonce(uint8* Nonce, uint32 Length)
+{
+//@todo joeg -- switch to CryptGenRandom() if possible or something equivalent
+	// Loop through generating a random value for each byte
+	for (uint32 NonceIndex = 0; NonceIndex < Length; NonceIndex++)
+	{
+		Nonce[NonceIndex] = (uint8)(FMath::Rand() & 255);
+	}
+}
+
+/**
+ * Environment for the current online platform
+ */
+namespace EOnlineEnvironment
+{
+	enum Type
+	{
+		/** Dev environment */
+		Development,
+		/** Cert environment */
+		Certification,
+		/** Prod environment */
+		Production,
+		/** Not determined yet */
+		Unknown
+	};
+
+	/** @return the stringified version of the enum passed in */
+	inline const TCHAR* ToString(EOnlineEnvironment::Type EnvironmentType)
+	{
+		switch (EnvironmentType)
+		{
+			case Development: return TEXT("Development");
+			case Certification: return TEXT("Certification");
+			case Production: return TEXT("Production");
+			case Unknown: default: return TEXT("Unknown");
+		};
+	}
+}
+
+/** Possible login states */
+namespace ELoginStatus
+{
+	enum Type
+	{
+		/** Player has not logged in or chosen a local profile */
+		NotLoggedIn,
+		/** Player is using a local profile but is not logged in */
+		UsingLocalProfile,
+		/** Player has been validated by the platform specific authentication service */
+		LoggedIn
+	};
+
+	/** @return the stringified version of the enum passed in */
+	inline const TCHAR* ToString(ELoginStatus::Type EnumVal)
+	{
+		switch (EnumVal)
+		{
+			case NotLoggedIn:
+			{
+				return TEXT("NotLoggedIn");
+			}
+			case UsingLocalProfile:
+			{
+				return TEXT("UsingLocalProfile");
+			}
+			case LoggedIn:
+			{
+				return TEXT("LoggedIn");
+			}
+		}
+		return TEXT("");
+	}
+};
+
+/** Possible connection states */
+namespace EOnlineServerConnectionStatus
+{
+	enum Type
+	{
+		/** System normal (used for default state) */
+		Normal,
+		/** Gracefully disconnected from the online servers */
+		NotConnected,
+		/** Connected to the online servers just fine */
+		Connected,
+		/** Connection was lost for some reason */
+		ConnectionDropped,
+		/** Can't connect because of missing network connection */
+		NoNetworkConnection,
+		/** Service is temporarily unavailable */
+		ServiceUnavailable,
+		/** An update is required before connecting is possible */
+		UpdateRequired,
+		/** Servers are too busy to handle the request right now */
+		ServersTooBusy,
+		/** Disconnected due to duplicate login */
+		DuplicateLoginDetected,
+		/** Can't connect because of an invalid/unknown user */
+		InvalidUser,
+		/** Not authorized */
+		NotAuthorized,
+		/** Session has been lost on the backend */
+		InvalidSession
+	};
+
+	/** @return the stringified version of the enum passed in */
+	inline const TCHAR* ToString(EOnlineServerConnectionStatus::Type EnumVal)
+	{
+		switch (EnumVal)
+		{
+			case Normal:
+			{
+				return TEXT("Normal");
+			}
+			case NotConnected:
+			{
+				return TEXT("NotConnected");
+			}
+			case Connected:
+			{
+				return TEXT("Connected");
+			}
+			case ConnectionDropped:
+			{
+				return TEXT("ConnectionDropped");
+			}
+			case NoNetworkConnection:
+			{
+				return TEXT("NoNetworkConnection");
+			}
+			case ServiceUnavailable:
+			{
+				return TEXT("ServiceUnavailable");
+			}
+			case UpdateRequired:
+			{
+				return TEXT("UpdateRequired");
+			}
+			case ServersTooBusy:
+			{
+				return TEXT("ServersTooBusy");
+			}
+			case DuplicateLoginDetected:
+			{
+				return TEXT("DuplicateLoginDetected");
+			}
+			case InvalidUser:
+			{
+				return TEXT("InvalidUser");
+			}
+			case NotAuthorized:
+			{
+				return TEXT("Not Authorized");
+			}
+			case InvalidSession:
+			{
+				return TEXT("Invalid Session");
+			}
+		}
+		return TEXT("");
+	}
+};
+
+/** Possible feature privilege access levels */
+namespace EFeaturePrivilegeLevel
+{
+	enum Type
+	{
+		/** Not defined for the platform service */
+		Undefined,
+		/** Parental controls have disabled this feature */
+		Disabled,
+		/** Parental controls allow this feature only with people on their friends list */
+		EnabledFriendsOnly,
+		/** Parental controls allow this feature everywhere */
+		Enabled
+	};
+
+	/** @return the stringified version of the enum passed in */
+	inline const TCHAR* ToString(EFeaturePrivilegeLevel::Type EnumVal)
+	{
+		switch (EnumVal)
+		{
+			case Undefined:
+			{
+				return TEXT("Undefined");
+			}
+			case Disabled:
+			{
+				return TEXT("Disabled");
+			}
+			case EnabledFriendsOnly:
+			{
+				return TEXT("EnabledFriendsOnly");
+			}
+			case Enabled:
+			{
+				return TEXT("Enabled");
+			}
+		}
+		return TEXT("");
+	}
+};
+
+/** The state of an async task (read friends, read content, write cloud file, etc) request */
+namespace EOnlineAsyncTaskState
+{
+	enum Type
+	{
+		/** The task has not been started */
+		NotStarted,
+		/** The task is currently being processed */
+		InProgress,
+		/** The task has completed successfully */
+		Done,
+		/** The task failed to complete */
+		Failed
+	};  
+
+	/** @return the stringified version of the enum passed in */
+	inline const TCHAR* ToString(EOnlineAsyncTaskState::Type EnumVal)
+	{
+		switch (EnumVal)
+		{
+			case NotStarted:
+			{
+				return TEXT("NotStarted");
+			}
+			case InProgress:
+			{
+				return TEXT("InProgress");
+			}
+			case Done:
+			{
+				return TEXT("Done");
+			}
+			case Failed:
+			{
+				return TEXT("Failed");
+			}
+		}
+		return TEXT("");
+	}
+};
+
+/** The possible friend states for a friend entry */
+namespace EOnlineFriendState
+{
+	enum Type
+	{
+		/** Not currently online */
+		Offline,
+		/** Signed in and online */
+		Online,
+		/** Signed in, online, and idle */
+		Away,
+		/** Signed in, online, and asks to be left alone */
+		Busy
+	};
+
+	/** @return the stringified version of the enum passed in */
+	inline const TCHAR* ToString(EOnlineFriendState::Type EnumVal)
+	{
+		switch (EnumVal)
+		{
+			case Offline:
+			{
+				return TEXT("Offline");
+			}
+			case Online:
+			{
+				return TEXT("Online");
+			}
+			case Away:
+			{
+				return TEXT("Away");
+			}
+			case Busy:
+			{
+				return TEXT("Busy");
+			}
+		}
+		return TEXT("");
+	}
+};
+
+/** Leaderboard entry sort types */
+namespace ELeaderboardSort
+{
+	enum Type
+	{
+		/** Don't sort at all */
+		None,
+		/** Sort ascending */
+		Ascending,
+		/** Sort descending */
+		Descending
+	};
+
+	/** @return the stringified version of the enum passed in */
+	inline const TCHAR* ToString(ELeaderboardSort::Type EnumVal)
+	{
+		switch (EnumVal)
+		{
+		case None:
+			{
+				return TEXT("None");
+			}
+		case Ascending:
+			{
+				return TEXT("Ascending");
+			}
+		case Descending:
+			{
+				return TEXT("Descending");
+			}
+		}
+		return TEXT("");
+	}
+};
+
+/** Leaderboard display format */
+namespace ELeaderboardFormat
+{
+	enum Type
+	{
+		/** A raw number */
+		Number,
+		/** Time, in seconds */
+		Seconds,
+		/** Time, in milliseconds */
+		Milliseconds
+	};
+
+	/** @return the stringified version of the enum passed in */
+	inline const TCHAR* ToString(ELeaderboardFormat::Type EnumVal)
+	{
+		switch (EnumVal)
+		{
+		case Number:
+			{
+				return TEXT("Number");
+			}
+		case Seconds:
+			{
+				return TEXT("Seconds");
+			}
+		case Milliseconds:
+			{
+				return TEXT("Milliseconds");
+			}
+		}
+		return TEXT("");
+	}
+};
+
+/** How to upload leaderboard score updates */
+namespace ELeaderboardUpdateMethod
+{
+	enum Type
+	{
+		/** If current leaderboard score is better than the uploaded one, keep the current one */
+		KeepBest,
+		/** Leaderboard score is always replaced with uploaded value */
+		Force
+	};
+
+	/** @return the stringified version of the enum passed in */
+	inline const TCHAR* ToString(ELeaderboardUpdateMethod::Type EnumVal)
+	{
+		switch (EnumVal)
+		{
+		case KeepBest:
+			{
+				return TEXT("KeepBest");
+			}
+		case Force:
+			{
+				return TEXT("Force");
+			}
+		}
+		return TEXT("");
+	}
+};
+
+/** Enum indicating the state the LAN beacon is in */
+namespace ELanBeaconState
+{
+	enum Type
+	{
+		/** The lan beacon is disabled */
+		NotUsingLanBeacon,
+		/** The lan beacon is responding to client requests for information */
+		Hosting,
+		/** The lan beacon is querying servers for information */
+		Searching
+	};
+
+	/** @return the stringified version of the enum passed in */
+	inline const TCHAR* ToString(ELanBeaconState::Type EnumVal)
+	{
+		switch (EnumVal)
+		{
+
+		case NotUsingLanBeacon:
+			{
+				return TEXT("NotUsingLanBeacon");
+			}
+		case Hosting:
+			{
+				return TEXT("Hosting");
+			}
+		case Searching:
+			{
+				return TEXT("Searching");
+			}
+		}
+		return TEXT("");
+	}
+};
+
+/** Enum indicating the current state of the online session (in progress, ended, etc.) */
+namespace EOnlineSessionState
+{
+	enum Type
+	{
+		/** An online session has not been created yet */
+		NoSession,
+		/** An online session is in the process of being created */
+		Creating,
+		/** Session has been created but the session hasn't started (pre match lobby) */
+		Pending,
+		/** Session has been asked to start (may take time due to communication with backend) */
+		Starting,
+		/** The current session has started. Sessions with join in progress disabled are no longer joinable */
+		InProgress,
+		/** The session is still valid, but the session is no longer being played (post match lobby) */
+		Ending,
+		/** The session is closed and any stats committed */
+		Ended,
+		/** The session is being destroyed */
+		Destroying
+	};
+
+	/** @return the stringified version of the enum passed in */
+	inline const TCHAR* ToString(EOnlineSessionState::Type EnumVal)
+	{
+		switch (EnumVal)
+		{
+
+		case NoSession:
+			{
+				return TEXT("NoSession");
+			}
+		case Creating:
+			{
+				return TEXT("Creating");
+			}
+		case Pending:
+			{
+				return TEXT("Pending");
+			}
+		case Starting:
+			{
+				return TEXT("Starting");
+			}
+		case InProgress:
+			{
+				return TEXT("InProgress");
+			}
+		case Ending:
+			{
+				return TEXT("Ending");
+			}
+		case Ended:
+			{
+				return TEXT("Ended");
+			}
+		case Destroying:
+			{
+				return TEXT("Destroying");
+			}
+		}
+		return TEXT("");
+	}
+};
+
+/** The types of advertisement of settings to use */
+namespace EOnlineDataAdvertisementType
+{
+	enum Type
+	{
+		/** Don't advertise via the online service or QoS data */
+		DontAdvertise,
+		/** Advertise via the server ping data only */
+		ViaPingOnly,
+		/** Advertise via the online service only */
+		ViaOnlineService,
+		/** Advertise via the online service and via the ping data */
+		ViaOnlineServiceAndPing
+	};
+
+	/** @return the stringified version of the enum passed in */
+	inline const TCHAR* ToString(EOnlineDataAdvertisementType::Type EnumVal)
+	{
+		switch (EnumVal)
+		{
+		case DontAdvertise:
+			{
+				return TEXT("DontAdvertise");
+			}
+		case ViaPingOnly:
+			{
+				return TEXT("ViaPingOnly");
+			}
+		case ViaOnlineService:
+			{
+				return TEXT("OnlineService");
+			}
+		case ViaOnlineServiceAndPing:
+			{
+				return TEXT("OnlineServiceAndPing");
+			}
+		}
+		return TEXT("");
+	}
+}
+
+/** The types of comparison operations for a given search query */
+namespace EOnlineComparisonOp
+{
+	enum Type
+	{
+		Equals,
+		NotEquals,
+		GreaterThan,
+		GreaterThanEquals,
+		LessThan,
+		LessThanEquals,
+		Near,
+		In,
+		NotIn
+	};
+
+	/** @return the stringified version of the enum passed in */
+	inline const TCHAR* ToString(EOnlineComparisonOp::Type EnumVal)
+	{
+		switch (EnumVal)
+		{
+		case Equals:
+			{
+				return TEXT("Equals");
+			}
+		case NotEquals:
+			{
+				return TEXT("NotEquals");
+			}
+		case GreaterThan:
+			{
+				return TEXT("GreaterThan");
+			}
+		case GreaterThanEquals:
+			{
+				return TEXT("GreaterThanEquals");
+			}
+		case LessThan:
+			{
+				return TEXT("LessThan");
+			}
+		case LessThanEquals:
+			{
+				return TEXT("LessThanEquals");
+			}
+		case Near:
+			{
+				return TEXT("Near");
+			}
+		case In:
+			{
+				return TEXT("In");
+			}
+		case NotIn:
+			{
+				return TEXT("NotIn");
+			}
+		}
+		return TEXT("");
+	}
+}
+
+/** Return codes for the GetCached functions in the various subsystems. */
+namespace EOnlineCachedResult
+{
+	enum Type
+	{
+		Success, /** The requested data was found and returned successfully. */
+		NotFound /** The requested data was not found in the cache, and the out parameter was not modified. */
+	};
+
+	/**
+	 * @param EnumVal the enum to convert to a string
+	 * @return the stringified version of the enum passed in
+	 */
+	inline const TCHAR* ToString(EOnlineCachedResult::Type EnumVal)
+	{
+		switch (EnumVal)
+		{
+		case Success:
+			{
+				return TEXT("Success");
+			}
+		case NotFound:
+			{
+				return TEXT("NotFound");
+			}
+		}
+		return TEXT("");
+	}
+}
+
+/*
+ *	Base class for anything meant to be opaque so that the data can be passed around 
+ *  without consideration for the data it contains.
+ *	A human readable version of the data is available via the ToString() function
+ *	Otherwise, nothing but platform code should try to operate directly on the data
+ */
+class IOnlinePlatformData
+{
+protected:
+
+	/** Hidden on purpose */
+	IOnlinePlatformData()
+	{
+	}
+
+	/** Hidden on purpose */
+	IOnlinePlatformData(const IOnlinePlatformData& Src)
+	{
+	}
+
+	/** Hidden on purpose */
+	IOnlinePlatformData& operator=(const IOnlinePlatformData& Src)
+	{
+		return *this;
+	}
+
+	virtual bool Compare(const IOnlinePlatformData& Other) const
+	{
+		return (GetSize() == Other.GetSize()) &&
+			(FMemory::Memcmp(GetBytes(), Other.GetBytes(), GetSize()) == 0);
+	}
+
+public:
+
+	virtual ~IOnlinePlatformData() {}
+
+	/**
+	 *	Comparison operator
+	 */
+	bool operator==(const IOnlinePlatformData& Other) const
+	{
+		return Other.Compare(*this);
+	}
+
+	bool operator!=(const IOnlinePlatformData& Other) const
+	{
+		return !(IOnlinePlatformData::operator==(Other));
+	}
+	
+	/** 
+	 * Get the raw byte representation of this opaque data
+	 * This data is platform dependent and shouldn't be manipulated directly
+	 *
+	 * @return byte array of size GetSize()
+	 */
+	virtual const uint8* GetBytes() const = 0;
+
+	/** 
+	 * Get the size of the opaque data
+	 *
+	 * @return size in bytes of the data representation
+	 */
+	virtual int32 GetSize() const = 0;
+
+	/** 
+	 * Check the validity of the opaque data
+	 *
+	 * @return true if this is well formed data, false otherwise
+	 */
+	virtual bool IsValid() const = 0;
+
+	/** 
+	 * Platform specific conversion to string representation of data
+	 *
+	 * @return data in string form 
+	 */
+	virtual FString ToString() const = 0;
+
+	/** 
+	 * Get a human readable representation of the opaque data
+	 * Shouldn't be used for anything other than logging/debugging
+	 *
+	 * @return data in string form 
+	 */
+	virtual FString ToDebugString() const = 0;
+};
+
+/**
+ * TArray helper for IndexOfByPredicate() function
+ */
+struct FUniqueNetIdMatcher
+{
+private:
+	/** Target for comparison in the TArray */
+	const FUniqueNetId& UniqueIdTarget;
+
+public:
+	FUniqueNetIdMatcher(const FUniqueNetId& InUniqueIdTarget) :
+		UniqueIdTarget(InUniqueIdTarget)
+	{
+	}
+
+	/**
+	 * Match a given unique Id against the one stored in this struct
+	 *
+	 * @return true if they are an exact match, false otherwise
+	 */
+	bool operator()(const FUniqueNetId& Candidate) const
+	{
+		return UniqueIdTarget == Candidate;
+	}
+ 
+	/**
+	 * Match a given unique Id against the one stored in this struct
+	 *
+	 * @return true if they are an exact match, false otherwise
+	 */
+	bool operator()(const TSharedPtr<const FUniqueNetId>& Candidate) const
+	{
+		return UniqueIdTarget == *Candidate;
+	}
+
+	/**
+	 * Match a given unique Id against the one stored in this struct
+	 *
+	 * @return true if they are an exact match, false otherwise
+	 */
+	bool operator()(const TSharedRef<const FUniqueNetId>& Candidate) const
+	{
+		return UniqueIdTarget == *Candidate;
+	}
+};
+
+/**
+ * Unique net id wrapper for a string
+ */
+class FUniqueNetIdString : public FUniqueNetId
+{
+public:
+	/** Holds the net id for a player */
+	FString UniqueNetIdStr;
+
+	/** Default constructor */
+	FUniqueNetIdString()
+	{
+	}
+
+	/**
+	 * Constructs this object with the specified net id
+	 *
+	 * @param InUniqueNetId the id to set ours to
+	 */
+	explicit FUniqueNetIdString(const FString& InUniqueNetId) 
+		: UniqueNetIdStr(InUniqueNetId)
+	{
+	}
+
+	/**
+	 * Copy Constructor
+	 *
+	 * @param Src the id to copy
+	 */
+	explicit FUniqueNetIdString(const FUniqueNetId& Src) 
+		: UniqueNetIdStr(Src.ToString())
+	{
+	}
+
+	/**
+	 * Copy Constructor
+	 *
+	 * @param Src the id to copy
+	 */
+	explicit FUniqueNetIdString(const FUniqueNetIdString& Src) 
+		: UniqueNetIdStr(Src.UniqueNetIdStr)
+	{
+	}
+
+	// IOnlinePlatformData
+
+	virtual const uint8* GetBytes() const override
+	{
+		return (const uint8*)UniqueNetIdStr.GetCharArray().GetData();
+	}
+
+	virtual int32 GetSize() const override
+	{
+		return UniqueNetIdStr.GetCharArray().GetTypeSize() * UniqueNetIdStr.GetCharArray().Num();
+	}
+
+	virtual bool IsValid() const override
+	{
+		return !UniqueNetIdStr.IsEmpty();
+	}
+
+	virtual FString ToString() const override
+	{
+		return UniqueNetIdStr;
+	}
+
+	virtual FString ToDebugString() const override
+	{
+		return UniqueNetIdStr;
+	}
+
+	/** Needed for TMap::GetTypeHash() */
+	friend uint32 GetTypeHash(const FUniqueNetIdString& A)
+	{
+		return ::GetTypeHash(A.UniqueNetIdStr);
+	}
+};
+
+/** 
+ * Abstraction of a profile service shared file handle
+ * The class is meant to be opaque (see IOnlinePlatformData)
+ */
+class FSharedContentHandle : public IOnlinePlatformData
+{
+protected:
+
+	/** Hidden on purpose */
+	FSharedContentHandle()
+	{
+	}
+
+	/** Hidden on purpose */
+	FSharedContentHandle(const FSharedContentHandle& Src)
+	{
+	}
+
+	/** Hidden on purpose */
+	FSharedContentHandle& operator=(const FSharedContentHandle& Src)
+	{
+		return *this;
+	}
+
+public:
+
+	virtual ~FSharedContentHandle() {}
+};
+
+/** 
+ * Abstraction of a session's platform specific info
+ * The class is meant to be opaque (see IOnlinePlatformData)
+ */
+class FOnlineSessionInfo : public IOnlinePlatformData
+{
+protected:
+
+	/** Hidden on purpose */
+	FOnlineSessionInfo()
+	{
+	}
+
+	/** Hidden on purpose */
+	FOnlineSessionInfo(const FOnlineSessionInfo& Src)
+	{
+	}
+
+	/** Hidden on purpose */
+	FOnlineSessionInfo& operator=(const FOnlineSessionInfo& Src)
+	{
+		return *this;
+	}
+
+public:
+
+	virtual ~FOnlineSessionInfo() {}
+
+	/**
+	 * Get the session id associated with this session
+	 *
+	 * @return session id for this session
+	 */
+	virtual const FUniqueNetId& GetSessionId() const = 0;
+};
+
+/**
+ * Paging info needed for a request that can return paged results
+ */
+class FPagedQuery
+{
+public:
+	FPagedQuery(int32 InStart = 0, int32 InCount = -1)
+		: Start(InStart)
+		, Count(InCount)
+	{}
+
+	/** @return true if valid range */
+	bool IsValidRange() const
+	{
+		return Start >= 0 && Count >= 0;
+	}
+
+	/** first entry to fetch */
+	int32 Start;
+	/** total entries to fetch. -1 means ALL */
+	int32 Count;
+};
+
+/**
+ * Info for a response with paged results
+ */
+class FOnlinePagedResult
+{
+public:
+	FOnlinePagedResult()
+		: Start(0)
+		, Count(0)
+		, Total(0)
+	{}
+	virtual ~FOnlinePagedResult() {}
+
+	/** Starting entry */
+	int32 Start;
+	/** Number returned */
+	int32 Count;
+	/** Total available */
+	int32 Total;
+};
+
+/** Locale and country code */
+class FRegionInfo
+{
+public:
+
+	FRegionInfo(const FString& InCountry = FString(), const FString& InLocale = FString())
+		: Country(InCountry)
+		, Locale(InLocale)
+	{}
+
+	/** country code used for configuring things like currency/pricing specific to a country. eg. US */
+	FString Country;
+	/** local code used to select the localization language. eg. en_US */
+	FString Locale;
+};
+
+/** Holds metadata about a given downloadable file */
+struct FCloudFileHeader
+{	
+	/** Hash value, if applicable, of the given file contents */
+    FString Hash;
+	/** The hash algorithm used to sign this file */
+	FName HashType;
+	/** Filename as downloaded */
+    FString DLName;
+	/** Logical filename, maps to the downloaded filename */
+    FString FileName;
+	/** File size */
+    int32 FileSize;
+	/** The full URL to download the file if it is stored in a CDN or separate host site */
+	FString URL;
+
+    /** Constructors */
+    FCloudFileHeader() :
+		FileSize(0)
+	{}
+
+	FCloudFileHeader(const FString& InFileName, const FString& InDLName, int32 InFileSize) :
+		DLName(InDLName),
+		FileName(InFileName),
+		FileSize(InFileSize)
+	{}
+
+	bool operator==(const FCloudFileHeader& Other) const
+	{
+		return FileSize == Other.FileSize &&
+			Hash == Other.Hash &&
+			HashType == Other.HashType &&
+			DLName == Other.DLName &&
+			FileName == Other.FileName &&
+			URL == Other.URL;
+	}
+
+	bool operator<(const FCloudFileHeader& Other) const
+	{
+		return FileName.Compare(Other.FileName, ESearchCase::IgnoreCase) < 0;
+	}
+};
+
+/** Holds the data used in downloading a file asynchronously from the online service */
+struct FCloudFile
+{
+	/** The name of the file as requested */
+	FString FileName;
+	/** The async state the file download is in */
+	EOnlineAsyncTaskState::Type AsyncState;
+	/** The buffer of data for the file */
+	TArray<uint8> Data;
+
+	/** Constructors */
+	FCloudFile() :
+		AsyncState(EOnlineAsyncTaskState::NotStarted)
+	{
+	}
+
+	FCloudFile(const FString& InFileName) :
+		FileName(InFileName),
+		AsyncState(EOnlineAsyncTaskState::NotStarted)
+	{
+	}
+
+	virtual ~FCloudFile() {}
+};
+
+/**
+ * Base for all online user info
+ */
+class FOnlineUser
+{
+public:
+	/**
+	 * destructor
+	 */
+	virtual ~FOnlineUser() {}
+
+	/** 
+	 * @return Id associated with the user account provided by the online service during registration 
+	 */
+	virtual TSharedRef<const FUniqueNetId> GetUserId() const = 0;
+	/**
+	 * @return the real name for the user if known
+	 */
+	virtual FString GetRealName() const = 0;
+	/**
+	 * @return the nickname of the user if known
+	 */
+	virtual FString GetDisplayName(const FString& Platform = FString()) const = 0;
+	/** 
+	 * @return Any additional user data associated with a registered user
+	 */
+	virtual bool GetUserAttribute(const FString& AttrName, FString& OutAttrValue) const = 0;
+};
+
+/**
+ * User account information returned via IOnlineIdentity interface
+ */
+class FUserOnlineAccount : public FOnlineUser
+{
+public:
+	/**
+	 * @return Access token which is provided to user once authenticated by the online service
+	 */
+	virtual FString GetAccessToken() const = 0;
+	/** 
+	 * @return Any additional auth data associated with a registered user
+	 */
+	virtual bool GetAuthAttribute(const FString& AttrName, FString& OutAttrValue) const = 0;
+	/** 
+	 * @return True, if the data has been changed
+	 */
+	virtual bool SetUserAttribute(const FString& AttrName, const FString& AttrValue) = 0;
+};
+
+/** 
+ * Friend list invite states 
+ */
+namespace EInviteStatus
+{
+	enum Type
+	{
+		/** unknown state */
+		Unknown,
+		/** Friend has accepted the invite */
+		Accepted,
+		/** Friend has sent player an invite, but it has not been accepted/rejected */
+		PendingInbound,
+		/** Player has sent friend an invite, but it has not been accepted/rejected */
+		PendingOutbound,
+		/** Player has been blocked */
+		Blocked,
+	};
+
+	/** 
+	 * @return the stringified version of the enum passed in 
+	 */
+	inline const TCHAR* ToString(EInviteStatus::Type EnumVal)
+	{
+		switch (EnumVal)
+		{
+			case Unknown:
+			{
+				return TEXT("Unknown");
+			}
+			case Accepted:
+			{
+				return TEXT("Accepted");
+			}
+			case PendingInbound:
+			{
+				return TEXT("PendingInbound");
+			}
+			case PendingOutbound:
+			{
+				return TEXT("PendingOutbound");
+			}
+			case Blocked:
+			{
+				return TEXT("Blocked");
+			}
+		}
+		return TEXT("");
+	}
+};
+
+/**
+ * Friend user info returned via IOnlineFriends interface
+ */
+class FOnlineFriend : public FOnlineUser
+{
+public:
+	
+	/**
+	 * @return the current invite status of a friend wrt to user that queried
+	 */
+	virtual EInviteStatus::Type GetInviteStatus() const = 0;
+	
+	/**
+	 * @return presence info for an online friend
+	 */
+	virtual const class FOnlineUserPresence& GetPresence() const = 0;
+};
+
+/**
+ * Recent player user info returned via IOnlineFriends interface
+ */
+class FOnlineRecentPlayer : public FOnlineUser
+{
+public:
+	
+	/**
+	 * @return last time the player was seen by the current user
+	 */
+	virtual FDateTime GetLastSeen() const = 0;
+};
+
+/**
+ * Blocked user info returned via IOnlineFriends interface
+ */
+class FOnlineBlockedPlayer : public FOnlineUser
+{
+};
+
+/** The possible permission categories we can choose from to read from the server */
+namespace EOnlineSharingReadCategory
+{
+	enum Type
+	{
+		None			= 0x00,
+		// Read access to posts on the users feeds
+		Posts			= 0x01,
+		// Read access for a users friend information, and all data about those friends. e.g. Friends List and Individual Friends Birthday
+		Friends			= 0x02,
+		// Read access to a users mailbox
+		Mailbox			= 0x04,
+		// Read the current online status of a user
+		OnlineStatus	= 0x08,
+		// Read a users profile information, e.g. Users Birthday
+		ProfileInfo		= 0x10,	
+		// Read information about the users locations and location history
+		LocationInfo	= 0x20,
+
+		Default			= ProfileInfo|LocationInfo,
+	};
+
+
+
+	/** @return the stringified version of the enum passed in */
+	inline const TCHAR* ToString(EOnlineSharingReadCategory::Type CategoryType)
+	{
+		switch (CategoryType)
+		{
+		case None:
+			{
+				return TEXT("Category undefined");
+			}
+		case Posts:
+			{
+				return TEXT("Posts");
+			}
+		case Friends:
+			{
+				return TEXT("Friends");
+			}
+		case Mailbox:
+			{
+				return TEXT("Mailbox");
+			}
+		case OnlineStatus:
+			{
+				return TEXT("Online Status");
+			}
+		case ProfileInfo:
+			{
+				return TEXT("Profile Information");
+			}
+		case LocationInfo:
+			{
+				return TEXT("Location Information");
+			}
+		}
+		return TEXT("");
+	}
+}
+
+
+/** The possible permission categories we can choose from to publish to the server */
+namespace EOnlineSharingPublishingCategory
+{
+	enum Type
+	{
+		None			= 0x00,
+		// Permission to post to a users news feed
+		Posts			= 0x01,
+		// Permission to manage a users friends list. Add/Remove contacts
+		Friends			= 0x02,
+		// Manage a users account settings, such as pages they subscribe to, or which notifications they receive
+		AccountAdmin	= 0x04,
+		// Manage a users events. This features the capacity to create events as well as respond to events.
+		Events			= 0x08,
+
+		Default			= None,
+	};
+
+
+	/** @return the stringified version of the enum passed in */
+	inline const TCHAR* ToString(EOnlineSharingPublishingCategory::Type CategoryType)
+	{
+		switch (CategoryType)
+		{
+		case None:
+			{
+				return TEXT("Category undefined");
+			}
+		case Posts:
+			{
+				return TEXT("Posts");
+			}
+		case Friends:
+			{
+				return TEXT("Friends");
+			}
+		case AccountAdmin:
+			{
+				return TEXT("Account Admin");
+			}
+		case Events:
+			{
+				return TEXT("Events");
+			}
+		}
+		return TEXT("");
+	}
+}
+
+
+/** Privacy permissions used for Online Status updates */
+namespace EOnlineStatusUpdatePrivacy
+{
+	enum Type
+	{
+		OnlyMe,			// Post will only be visible to the user alone
+		OnlyFriends,	// Post will only be visible to the user and the users friends
+		Everyone,		// Post will be visible to everyone
+	};
+
+	inline const TCHAR* ToString(EOnlineStatusUpdatePrivacy::Type PrivacyType)
+	{
+		switch (PrivacyType)
+		{
+		case OnlyMe:
+			return TEXT("Only Me");
+		case OnlyFriends:
+			return TEXT("Only Friends");
+		case Everyone:
+			return TEXT("Everyone");
+		}
+	}
+}
+
+/**
+* unique identifier for notification transports
+*/
+typedef FString FNotificationTransportId;
+
+/**
+* Id of a party instance
+*/
+class FOnlinePartyId : public IOnlinePlatformData, public TSharedFromThis<FOnlinePartyId>
+{
+protected:
+	/** Hidden on purpose */
+	FOnlinePartyId()
+	{
+	}
+
+	/** Hidden on purpose */
+	FOnlinePartyId(const FOnlinePartyId& Src)
+	{
+	}
+
+	/** Hidden on purpose */
+	FOnlinePartyId& operator=(const FOnlinePartyId& Src)
+	{
+		return *this;
+	}
+
+public:
+	virtual ~FOnlinePartyId() {}
+};
+
+/**
+* Id of a party's type
+*/
+class FOnlinePartyTypeId
+{
+public:
+	typedef uint32 TInternalType;
+	explicit FOnlinePartyTypeId(const TInternalType InValue = 0) : Value(InValue) {}
+	FOnlinePartyTypeId(const FOnlinePartyTypeId& Other) : Value(Other.Value) {}
+
+	bool operator==(const FOnlinePartyTypeId Rhs) const { return Value == Rhs.Value; }
+	bool operator!=(const FOnlinePartyTypeId Rhs) const { return Value != Rhs.Value; }
+
+	TInternalType GetValue() const { return Value; }
+
+	friend bool IsValid(const FOnlinePartyTypeId Value);
+
+protected:
+	TInternalType Value;
+};
+
+inline bool IsValid(const FOnlinePartyTypeId Id) { return Id.GetValue() != 0; }
+
+inline uint32 GetTypeHash(const FOnlinePartyTypeId Id)
+{
+	return Id.GetValue();
+}
diff --git a/Engine/Plugins/Online/OnlineSubsystemAmazon/Source/OnlineSubsystemAmazon.Build.cs b/Engine/Plugins/Online/OnlineSubsystemAmazon/Source/OnlineSubsystemAmazon.Build.cs
new file mode 100644
--- /dev/null
+++ b/Engine/Plugins/Online/OnlineSubsystemAmazon/Source/OnlineSubsystemAmazon.Build.cs
@@ -0,0 +1,22 @@
+// Copyright 1998-2016 Epic Games, Inc. All Rights Reserved.
+
+using UnrealBuildTool;
+using System.IO;
+
+public class OnlineSubsystemAmazon : ModuleRules
+{
+	public OnlineSubsystemAmazon(TargetInfo Target)
+	{
+		Definitions.Add("ONLINESUBSYSTEMAMAZON_PACKAGE=1");		
+
+		PrivateDependencyModuleNames.AddRange(
+			new string[] {
+				"Core", 
+                "CoreUObject",
+				"HTTP",
+				"Json",
+                "OnlineSubsystem", 
+			}
+		);
+	}
+}
diff --git a/Engine/Plugins/Online/OnlineSubsystemAmazon/Source/Private/OnlineSubsystemAmazon.cpp b/Engine/Plugins/Online/OnlineSubsystemAmazon/Source/Private/OnlineSubsystemAmazon.cpp
new file mode 100644
--- /dev/null
+++ b/Engine/Plugins/Online/OnlineSubsystemAmazon/Source/Private/OnlineSubsystemAmazon.cpp
@@ -0,0 +1,144 @@
+// Copyright 1998-2016 Epic Games, Inc. All Rights Reserved.
+
+#include "OnlineSubsystemAmazonPCH.h"
+#include "OnlineSubsystemAmazon.h"
+#include "OnlineSubsystemAmazonModule.h"
+#include "OnlineIdentityAmazon.h"
+
+// FOnlineSubsystemAmazonModule
+IMPLEMENT_MODULE(FOnlineSubsystemAmazonModule, OnlineSubsystemAmazon);
+
+/**
+ * Class responsible for creating instance(s) of the subsystem
+ */
+class FOnlineFactoryAmazon : public IOnlineFactory
+{
+public:
+
+	FOnlineFactoryAmazon() {}
+	virtual ~FOnlineFactoryAmazon() {}
+
+	virtual IOnlineSubsystemPtr CreateSubsystem(FName InstanceName)
+	{
+		FOnlineSubsystemAmazonPtr OnlineSub = MakeShareable(new FOnlineSubsystemAmazon());
+		if (OnlineSub->IsEnabled())
+		{
+			if(!OnlineSub->Init())
+			{
+				UE_LOG(LogOnline, Warning, TEXT("Amazon API failed to initialize!"));
+				OnlineSub->Shutdown();
+				OnlineSub = NULL;
+			}
+		}
+		else
+		{
+			UE_LOG(LogOnline, Warning, TEXT("Amazon API disabled!"));
+			OnlineSub->Shutdown();
+			OnlineSub = NULL;
+		}
+
+		return OnlineSub;
+	}
+};
+
+void FOnlineSubsystemAmazonModule::StartupModule()
+{
+	UE_LOG(LogOnline, Log, TEXT("Amazon Startup!"));
+
+	AmazonFactory = new FOnlineFactoryAmazon();
+
+	// Create and register our singleton factory with the main online subsystem for easy access
+	FOnlineSubsystemModule& OSS = FModuleManager::GetModuleChecked<FOnlineSubsystemModule>("OnlineSubsystem");
+	OSS.RegisterPlatformService(AMAZON_SUBSYSTEM, AmazonFactory);
+}
+
+void FOnlineSubsystemAmazonModule::ShutdownModule()
+{
+	UE_LOG(LogOnline, Log, TEXT("Amazon Shutdown!"));
+
+	FOnlineSubsystemModule& OSS = FModuleManager::GetModuleChecked<FOnlineSubsystemModule>("OnlineSubsystem");
+	OSS.UnregisterPlatformService(AMAZON_SUBSYSTEM);
+	
+	delete AmazonFactory;
+	AmazonFactory = NULL;
+}
+
+IOnlineIdentityPtr FOnlineSubsystemAmazon::GetIdentityInterface() const
+{
+	return IdentityInterface;
+}
+
+
+bool FOnlineSubsystemAmazon::Tick(float DeltaTime)
+{
+	if (!FOnlineSubsystemImpl::Tick(DeltaTime))
+	{
+		return false;
+	}
+
+	if (IdentityInterface.IsValid())
+	{
+		TickToggle ^= 1;
+		IdentityInterface->Tick(DeltaTime, TickToggle);
+	}
+	return true;
+}
+
+bool FOnlineSubsystemAmazon::Init()
+{
+	IdentityInterface = MakeShareable(new FOnlineIdentityAmazon());
+	return true;
+}
+
+bool FOnlineSubsystemAmazon::Shutdown()
+{
+	UE_LOG(LogOnline, Display, TEXT("FOnlineSubsystemAmazon::Shutdown()"));
+	IdentityInterface = NULL;
+	FOnlineSubsystemImpl::Shutdown();
+	return true;
+}
+
+FString FOnlineSubsystemAmazon::GetAppId() const
+{
+	return TEXT("Amazon");
+}
+
+bool FOnlineSubsystemAmazon::Exec(UWorld* InWorld, const TCHAR* Cmd, FOutputDevice& Ar) 
+{
+	if (FOnlineSubsystemImpl::Exec(InWorld, Cmd, Ar))
+	{
+		return true;
+	}
+	return false;
+}
+
+bool FOnlineSubsystemAmazon::IsEnabled(void)
+{
+	// Check the ini for disabling Amazon
+	bool bEnableAmazon = true;
+	if (GConfig->GetBool(TEXT("OnlineSubsystemAmazon"), TEXT("bEnabled"), bEnableAmazon, GEngineIni) && 
+		bEnableAmazon)
+	{
+		// Check the commandline for disabling Amazon
+		bEnableAmazon = !FParse::Param(FCommandLine::Get(),TEXT("NOAMAZON"));
+#if UE_EDITOR
+		if (bEnableAmazon)
+		{
+			bEnableAmazon = IsRunningDedicatedServer() || IsRunningGame();
+		}
+#endif
+	}
+	return bEnableAmazon;
+}
+
+FOnlineSubsystemAmazon::FOnlineSubsystemAmazon() :
+	IdentityInterface(NULL),
+	TickToggle(0)
+{
+
+}
+
+FOnlineSubsystemAmazon::~FOnlineSubsystemAmazon()
+{
+
+}
diff --git a/Engine/Plugins/Online/OnlineSubsystemAmazon/Source/Private/OnlineSubsystemAmazonPCH.h b/Engine/Plugins/Online/OnlineSubsystemAmazon/Source/Private/OnlineSubsystemAmazonPCH.h
new file mode 100644
--- /dev/null
+++ b/Engine/Plugins/Online/OnlineSubsystemAmazon/Source/Private/OnlineSubsystemAmazonPCH.h
@@ -0,0 +1,12 @@
+// Copyright 1998-2016 Epic Games, Inc. All Rights Reserved.
+
+#pragma once
+
+#include "OnlineSubsystem.h"
+#include "OnlineSubsystemAmazon.h"
+
+/** pre-pended to all Amazon logging */
+#undef ONLINE_LOG_PREFIX
+#define ONLINE_LOG_PREFIX TEXT("Amazon: ")
+
+
diff --git a/Engine/Plugins/Online/OnlineSubsystemFacebook/Source/OnlineSubsystemFacebook.Build.cs b/Engine/Plugins/Online/OnlineSubsystemFacebook/Source/OnlineSubsystemFacebook.Build.cs
new file mode 100644
--- /dev/null
+++ b/Engine/Plugins/Online/OnlineSubsystemFacebook/Source/OnlineSubsystemFacebook.Build.cs
@@ -0,0 +1,39 @@
+// Copyright 1998-2016 Epic Games, Inc. All Rights Reserved.
+
+using UnrealBuildTool;
+
+public class OnlineSubsystemFacebook : ModuleRules
+{
+    public OnlineSubsystemFacebook(TargetInfo Target)
+	{
+		Definitions.Add("ONLINESUBSYSTEMFACEBOOK_PACKAGE=1");
+
+        PrivateIncludePaths.Add("Private");
+
+        PrivateDependencyModuleNames.AddRange(
+            new string[] { 
+                "Core",
+                "CoreUObject",
+                "HTTP",
+				"ImageCore",
+				"Json",
+                "OnlineSubsystem", 
+            }
+            );
+
+        if (Target.Platform == UnrealTargetPlatform.IOS)
+        {
+			AddEngineThirdPartyPrivateStaticDependencies(Target, "Facebook");
+
+            PrivateIncludePaths.Add("Private/IOS");
+        }
+        else if (Target.Platform == UnrealTargetPlatform.Win32 || Target.Platform == UnrealTargetPlatform.Win64)
+        {
+            PrivateIncludePaths.Add("Private/Windows");
+        }
+		else
+		{
+			PrecompileForTargets = PrecompileTargetsType.None;
+		}
+	}
+}
diff --git a/Engine/Plugins/Online/OnlineSubsystemFacebook/Source/Private/IOS/OnlineFriendsFacebook.cpp b/Engine/Plugins/Online/OnlineSubsystemFacebook/Source/Private/IOS/OnlineFriendsFacebook.cpp
new file mode 100644
--- /dev/null
+++ b/Engine/Plugins/Online/OnlineSubsystemFacebook/Source/Private/IOS/OnlineFriendsFacebook.cpp
@@ -0,0 +1,306 @@
+// Copyright 1998-2016 Epic Games, Inc. All Rights Reserved.
+
+
+#include "OnlineSubsystemFacebookPrivatePCH.h"
+#include "OnlineIdentityFacebook.h"
+#include "OnlineSharingFacebook.h"
+#include "OnlineFriendsFacebook.h"
+
+#import <FBSDKCoreKit/FBSDKCoreKit.h>
+
+// FOnlineFriendFacebook
+
+TSharedRef<const FUniqueNetId> FOnlineFriendFacebook::GetUserId() const
+{
+	return UserId;
+}
+
+FString FOnlineFriendFacebook::GetRealName() const
+{
+	FString Result;
+	GetAccountData(TEXT("name"), Result);
+	return Result;
+}
+
+FString FOnlineFriendFacebook::GetDisplayName(const FString& Platform) const
+{
+	FString Result;
+	GetAccountData(TEXT("username"), Result);
+	return Result;
+}
+
+bool FOnlineFriendFacebook::GetUserAttribute(const FString& AttrName, FString& OutAttrValue) const
+{
+	return GetAccountData(AttrName, OutAttrValue);
+}
+
+EInviteStatus::Type FOnlineFriendFacebook::GetInviteStatus() const
+{
+	return EInviteStatus::Accepted;
+}
+
+const FOnlineUserPresence& FOnlineFriendFacebook::GetPresence() const
+{
+	return Presence;
+}
+
+// FOnlineFriendsFacebook
+
+FOnlineFriendsFacebook::FOnlineFriendsFacebook(FOnlineSubsystemFacebook* InSubsystem) 
+{
+	// Get our handle to the identity interface
+	IdentityInterface = InSubsystem->GetIdentityInterface();
+	check( IdentityInterface.IsValid() );
+	SharingInterface = InSubsystem->GetSharingInterface();
+	check( SharingInterface.IsValid() );
+	GConfig->GetArray(TEXT("OnlineSubsystemFacebook.OnlineFriendsFacebook"), TEXT("FriendsFields"), FriendsFields, GEngineIni);	
+	
+	// always required fields
+	FriendsFields.AddUnique(TEXT("name"));
+}
+
+
+FOnlineFriendsFacebook::FOnlineFriendsFacebook()
+{
+}
+
+
+FOnlineFriendsFacebook::~FOnlineFriendsFacebook()
+{
+	IdentityInterface = NULL;
+	SharingInterface = NULL;
+}
+
+bool FOnlineFriendsFacebook::ReadFriendsList(int32 LocalUserNum, const FString& ListName, const FOnReadFriendsListComplete& Delegate /*= FOnReadFriendsListComplete()*/)
+{
+	bool bRequestTriggered = false;
+	UE_LOG(LogOnline, Verbose, TEXT("FOnlineFriendsFacebook::ReadFriendsList()"));
+
+	// To read friends, we need to be logged in and have authorization.
+	if( IdentityInterface->GetLoginStatus( LocalUserNum ) == ELoginStatus::LoggedIn )
+	{
+		bRequestTriggered = true;
+
+		auto RequestFriendsReadPermissionsDelegate = FOnRequestNewReadPermissionsCompleteDelegate::CreateRaw(this, &FOnlineFriendsFacebook::OnReadFriendsPermissionsUpdated);
+		RequestFriendsReadPermissionsDelegateHandle = SharingInterface->AddOnRequestNewReadPermissionsCompleteDelegate_Handle(LocalUserNum, RequestFriendsReadPermissionsDelegate);
+		SharingInterface->RequestNewReadPermissions(LocalUserNum, EOnlineSharingReadCategory::Friends);
+		OnReadFriendsListCompleteDelegate = Delegate;
+	}
+	else
+	{
+		UE_LOG(LogOnline, Verbose, TEXT("Cannot read friends if we are not logged into facebook."));
+
+		// Cannot read friends if we are not logged into facebook.
+		Delegate.ExecuteIfBound(LocalUserNum, false, ListName, TEXT("not logged in."));
+	}
+
+	return bRequestTriggered;
+}
+
+bool FOnlineFriendsFacebook::DeleteFriendsList(int32 LocalUserNum, const FString& ListName, const FOnDeleteFriendsListComplete& Delegate /*= FOnDeleteFriendsListComplete()*/)
+{
+	Delegate.ExecuteIfBound(LocalUserNum, false, ListName, FString(TEXT("DeleteFriendsList() is not supported")));
+	return false;
+}
+
+bool FOnlineFriendsFacebook::SendInvite(int32 LocalUserNum, const FUniqueNetId& FriendId, const FString& ListName, const FOnSendInviteComplete& Delegate /*= FOnSendInviteComplete()*/)
+{
+	Delegate.ExecuteIfBound(LocalUserNum, false, FriendId, ListName, FString(TEXT("SendInvite() is not supported")));
+	return false;
+}
+
+bool FOnlineFriendsFacebook::AcceptInvite(int32 LocalUserNum, const FUniqueNetId& FriendId, const FString& ListName, const FOnAcceptInviteComplete& Delegate /*= FOnAcceptInviteComplete()*/)
+{
+	Delegate.ExecuteIfBound(LocalUserNum, false, FriendId, ListName, FString(TEXT("AcceptInvite() is not supported")));
+	return false;
+}
+
+bool FOnlineFriendsFacebook::RejectInvite(int32 LocalUserNum, const FUniqueNetId& FriendId, const FString& ListName)
+{
+	TriggerOnRejectInviteCompleteDelegates(LocalUserNum, false, FriendId, ListName, FString(TEXT("RejectInvite() is not supported")));
+	return false;
+}
+
+bool FOnlineFriendsFacebook::DeleteFriend(int32 LocalUserNum, const FUniqueNetId& FriendId, const FString& ListName)
+{
+	TriggerOnDeleteFriendCompleteDelegates(LocalUserNum, false, FriendId, ListName, FString(TEXT("DeleteFriend() is not supported")));
+	return false;
+}
+
+bool FOnlineFriendsFacebook::GetFriendsList(int32 LocalUserNum, const FString& ListName, TArray< TSharedRef<FOnlineFriend> >& OutFriends)
+{
+	UE_LOG(LogOnline, Verbose, TEXT("FOnlineFriendsFacebook::GetFriendsList()"));
+
+	for (int32 Idx=0; Idx < CachedFriends.Num(); Idx++)
+	{
+		OutFriends.Add(CachedFriends[Idx]);
+	}
+
+	return true;
+}
+
+TSharedPtr<FOnlineFriend> FOnlineFriendsFacebook::GetFriend(int32 LocalUserNum, const FUniqueNetId& FriendId, const FString& ListName)
+{
+	TSharedPtr<FOnlineFriend> Result;
+
+	UE_LOG(LogOnline, Verbose, TEXT("FOnlineFriendsFacebook::GetFriend()"));
+
+	for (int32 Idx=0; Idx < CachedFriends.Num(); Idx++)
+	{
+		if (*(CachedFriends[Idx]->GetUserId()) == FriendId)
+		{
+			Result = CachedFriends[Idx];
+			break;
+		}
+	}
+
+	return Result;
+}
+
+bool FOnlineFriendsFacebook::IsFriend(int32 LocalUserNum, const FUniqueNetId& FriendId, const FString& ListName)
+{
+	UE_LOG(LogOnline, Verbose, TEXT("FOnlineFriendsFacebook::IsFriend()"));
+
+	TSharedPtr<FOnlineFriend> Friend = GetFriend(LocalUserNum,FriendId,ListName);
+	if (Friend.IsValid() &&
+		Friend->GetInviteStatus() == EInviteStatus::Accepted)
+	{
+		return true;
+	}
+	return false;
+}
+
+bool FOnlineFriendsFacebook::QueryRecentPlayers(const FUniqueNetId& UserId, const FString& Namespace)
+{
+	UE_LOG(LogOnline, Verbose, TEXT("FOnlineFriendsFacebook::QueryRecentPlayers()"));
+
+	TriggerOnQueryRecentPlayersCompleteDelegates(UserId, Namespace, false, TEXT("not implemented"));
+
+	return false;
+}
+
+bool FOnlineFriendsFacebook::GetRecentPlayers(const FUniqueNetId& UserId, const FString& Namespace, TArray< TSharedRef<FOnlineRecentPlayer> >& OutRecentPlayers)
+{
+	return false;
+}
+
+bool FOnlineFriendsFacebook::BlockPlayer(int32 LocalUserNum, const FUniqueNetId& PlayerId)
+{
+	return false;
+}
+
+bool FOnlineFriendsFacebook::UnblockPlayer(int32 LocalUserNum, const FUniqueNetId& PlayerId)
+{
+	return false;
+}
+
+bool FOnlineFriendsFacebook::QueryBlockedPlayers(const FUniqueNetId& UserId)
+{
+	return false;
+}
+
+bool FOnlineFriendsFacebook::GetBlockedPlayers(const FUniqueNetId& UserId, TArray< TSharedRef<FOnlineBlockedPlayer> >& OutBlockedPlayers)
+{
+	return false;
+}
+
+void FOnlineFriendsFacebook::OnReadFriendsPermissionsUpdated(int32 LocalUserNum, bool bWasSuccessful)
+{
+	UE_LOG(LogOnline, Verbose, TEXT("FOnlineFriendsFacebook::OnReadPermissionsUpdated() - %d"), bWasSuccessful);
+	SharingInterface->ClearOnRequestNewReadPermissionsCompleteDelegate_Handle(LocalUserNum, RequestFriendsReadPermissionsDelegateHandle);
+
+	if( bWasSuccessful )
+	{
+		ReadFriendsUsingGraphPath(LocalUserNum, EFriendsLists::ToString(EFriendsLists::Default));
+	}
+	else
+	{
+		// Permissions werent applied so we cannot read friends.
+		OnReadFriendsListCompleteDelegate.ExecuteIfBound(LocalUserNum, false, EFriendsLists::ToString(EFriendsLists::Default), TEXT("No read permissions"));
+	}
+}
+
+void FOnlineFriendsFacebook::ReadFriendsUsingGraphPath(int32 LocalUserNum, const FString& ListName)
+{
+	UE_LOG(LogOnline, Verbose, TEXT("FOnlineFriendsFacebook::ReadFriendsUsingGraphPath()"));
+
+	dispatch_async(dispatch_get_main_queue(),^ 
+		{
+			// We need to determine what the graph path we are querying is.
+			// Facebook permissions only return friends who have our app installed.
+			NSMutableString* GraphPath = [[NSMutableString alloc] init];
+			[GraphPath appendString:@"/me/friends"];
+
+			NSMutableDictionary *parameters = [NSMutableDictionary dictionaryWithDictionary: @{@"limit" : @"50"}];
+			
+			// Optional list of fields to query for each friend
+			if(FriendsFields.Num() > 0)
+			{
+				FString FieldsStr;
+				for (int32 Idx=0; Idx < FriendsFields.Num(); Idx++)
+				{
+					FieldsStr += FriendsFields[Idx];
+					if (Idx < (FriendsFields.Num()-1))
+					{
+						FieldsStr += TEXT(",");
+					}
+				}
+				
+				[parameters setObject:[NSString stringWithFString:FieldsStr] forKey:@"fields"];
+			}
+
+			UE_LOG(LogOnline, Verbose, TEXT("GraphPath=%s"), *FString(GraphPath));
+
+			[[[FBSDKGraphRequest alloc] initWithGraphPath:GraphPath parameters:parameters HTTPMethod:@"GET"]
+				startWithCompletionHandler:^(FBSDKGraphRequestConnection *connection, id result, NSError *error)
+				{
+					bool bSuccess = error == nil;
+
+					if(error)
+					{
+			 			 UE_LOG(LogOnline, Display, TEXT("Error=%s"), *FString([error localizedDescription]));
+					}
+					else
+					{
+                        NSArray* friends = [[NSArray alloc] initWithArray:[result objectForKey:@"data"]];
+
+						UE_LOG(LogOnline, Display, TEXT("Found %i friends"), [friends count]);
+
+						// Clear our previosly cached friends before we repopulate the cache.
+						CachedFriends.Empty();
+
+						for( int32 FriendIdx = 0; FriendIdx < [friends count]; FriendIdx++ )
+						{	
+							NSDictionary* user = friends[ FriendIdx ];		
+							const FString Id([user objectForKey : @"objectID"]);
+							// Add new friend entry to list
+							TSharedRef<FOnlineFriendFacebook> FriendEntry(
+								new FOnlineFriendFacebook(*Id)
+								);
+							FriendEntry->AccountData.Add(
+								TEXT("id"), *Id
+								);
+							FriendEntry->AccountData.Add(
+								TEXT("name"), *FString([user objectForKey:@"name"])
+								);
+							FriendEntry->AccountData.Add(
+								TEXT("username"), *FString([user objectForKey:@"name"])
+								);
+							CachedFriends.Add(FriendEntry);
+
+							UE_LOG(LogOnline, Verbose, TEXT("GCFriend - Id:%s - NickName:%s - RealName:%s"), 
+								*FriendEntry->GetUserId()->ToString(), *FriendEntry->GetDisplayName(), *FriendEntry->GetRealName() );	
+						}
+					}
+
+					// Did this operation complete? Let whoever is listening know.
+					OnReadFriendsListCompleteDelegate.ExecuteIfBound(LocalUserNum, bSuccess, ListName, FString());
+				}
+			];
+		}
+	);
+}
+
+void FOnlineFriendsFacebook::DumpBlockedPlayers() const
+{
+}
\ No newline at end of file
diff --git a/Engine/Plugins/Online/OnlineSubsystemFacebook/Source/Private/IOS/OnlineFriendsFacebook.h b/Engine/Plugins/Online/OnlineSubsystemFacebook/Source/Private/IOS/OnlineFriendsFacebook.h
new file mode 100644
--- /dev/null
+++ b/Engine/Plugins/Online/OnlineSubsystemFacebook/Source/Private/IOS/OnlineFriendsFacebook.h
@@ -0,0 +1,170 @@
+// Copyright 1998-2016 Epic Games, Inc. All Rights Reserved.
+
+#pragma once
+
+#include "OnlineFriendsInterface.h"
+#include "OnlineSharingFacebook.h"
+#include "OnlineSubsystemFacebookPackage.h"
+#include "OnlinePresenceInterface.h"
+
+/**
+ * Info associated with an online friend on the facebook service
+ */
+class FOnlineFriendFacebook : 
+	public FOnlineFriend
+{
+public:
+
+	// FOnlineUser
+
+	virtual TSharedRef<const FUniqueNetId> GetUserId() const override;
+	virtual FString GetRealName() const override;
+	virtual FString GetDisplayName(const FString& Platform = FString()) const override;
+	virtual bool GetUserAttribute(const FString& AttrName, FString& OutAttrValue) const override;
+
+	// FOnlineFriend
+	
+	virtual EInviteStatus::Type GetInviteStatus() const override;
+	virtual const FOnlineUserPresence& GetPresence() const override;
+
+	// FOnlineFriendMcp
+
+	/**
+	 * Init/default constructor
+	 */
+	FOnlineFriendFacebook(const FString& InUserId=TEXT("")) 
+		: UserId(new FUniqueNetIdString(InUserId))
+	{
+	}
+
+	/**
+	 * Destructor
+	 */
+	virtual ~FOnlineFriendFacebook()
+	{
+	}
+
+	/**
+	 * Get account data attribute
+	 *
+	 * @param Key account data entry key
+	 * @param OutVal [out] value that was found
+	 *
+	 * @return true if entry was found
+	 */
+	inline bool GetAccountData(const FString& Key, FString& OutVal) const
+	{
+		const FString* FoundVal = AccountData.Find(Key);
+		if (FoundVal != NULL)
+		{
+			OutVal = *FoundVal;
+			return true;
+		}
+		return false;
+	}
+
+	/** User Id represented as a FUniqueNetId */
+	TSharedRef<const FUniqueNetId> UserId;
+	/** Any addition account data associated with the friend */
+	TMap<FString, FString> AccountData;
+	/** @temp presence info */
+	FOnlineUserPresence Presence;
+};
+
+
+/**
+ * Facebook service implementation of the online friends interface
+ */
+class FOnlineFriendsFacebook :
+	public IOnlineFriends
+{
+
+public:
+
+	// IOnlineFriends
+
+	virtual bool ReadFriendsList(int32 LocalUserNum, const FString& ListName, const FOnReadFriendsListComplete& Delegate = FOnReadFriendsListComplete()) override;
+	virtual bool DeleteFriendsList(int32 LocalUserNum, const FString& ListName, const FOnDeleteFriendsListComplete& Delegate = FOnDeleteFriendsListComplete()) override;
+	virtual bool SendInvite(int32 LocalUserNum, const FUniqueNetId& FriendId, const FString& ListName,  const FOnSendInviteComplete& Delegate = FOnSendInviteComplete()) override;
+	virtual bool AcceptInvite(int32 LocalUserNum, const FUniqueNetId& FriendId, const FString& ListName, const FOnAcceptInviteComplete& Delegate = FOnAcceptInviteComplete()) override;
+ 	virtual bool RejectInvite(int32 LocalUserNum, const FUniqueNetId& FriendId, const FString& ListName) override;
+ 	virtual bool DeleteFriend(int32 LocalUserNum, const FUniqueNetId& FriendId, const FString& ListName) override;
+	virtual bool GetFriendsList(int32 LocalUserNum, const FString& ListName, TArray< TSharedRef<FOnlineFriend> >& OutFriends) override;
+	virtual TSharedPtr<FOnlineFriend> GetFriend(int32 LocalUserNum, const FUniqueNetId& FriendId, const FString& ListName) override;
+	virtual bool IsFriend(int32 LocalUserNum, const FUniqueNetId& FriendId, const FString& ListName) override;
+	virtual bool QueryRecentPlayers(const FUniqueNetId& UserId, const FString& Namespace) override;
+	virtual bool GetRecentPlayers(const FUniqueNetId& UserId, const FString& Namespace, TArray< TSharedRef<FOnlineRecentPlayer> >& OutRecentPlayers) override;
+	virtual bool BlockPlayer(int32 LocalUserNum, const FUniqueNetId& PlayerId) override;
+	virtual bool UnblockPlayer(int32 LocalUserNum, const FUniqueNetId& PlayerId) override;
+	virtual bool QueryBlockedPlayers(const FUniqueNetId& UserId) override;
+	virtual bool GetBlockedPlayers(const FUniqueNetId& UserId, TArray< TSharedRef<FOnlineBlockedPlayer> >& OutBlockedPlayers) override;
+	virtual void DumpBlockedPlayers() const override;
+
+	// FOnlineFriendsFacebook
+
+PACKAGE_SCOPE:
+
+	/**
+	 * Constructor
+	 *
+	 * @param InSubsystem Facebook subsystem being used
+	 */
+	FOnlineFriendsFacebook(class FOnlineSubsystemFacebook* InSubsystem);
+
+
+public:
+	
+	/**
+	 * Destructor
+	 */
+	virtual ~FOnlineFriendsFacebook();
+
+
+private:
+
+	/**
+	 * Should use the OSS param constructor instead
+	 */
+	FOnlineFriendsFacebook();
+
+
+private:
+
+	/**
+	 * Delegate fired when we have a response after we have requested read permissions of the users friends list
+	 *
+	 * @param LocalUserNum		- The local player number the friends were read for
+	 * @param bWasSuccessful	- Whether the permissions were updated successfully
+	 */
+	void OnReadFriendsPermissionsUpdated(int32 LocalUserNum, bool bWasSuccessful);
+
+	/**
+	 * Using the open graph for fb to read friends list, called when we have ensured correct permissions
+	 *
+	 * @param LocalUserNum - The local player number the friends were read for
+	 * @param ListName - named list to read
+	 */
+	void ReadFriendsUsingGraphPath(int32 LocalUserNum, const FString& ListName);
+
+	/** Handle to delegate used to notify this interface that permissions are updated, and that we can now read the friends list */
+	FDelegateHandle RequestFriendsReadPermissionsDelegateHandle;
+
+	/** Delegate called when reading the friends list is completed */
+	FOnReadFriendsListComplete OnReadFriendsListCompleteDelegate;
+
+private:
+
+	/** Reference to the main facebook identity */
+	IOnlineIdentityPtr IdentityInterface;
+
+	/** Reference to the main facebook sharing interface */
+	IOnlineSharingPtr SharingInterface;
+
+	/** The collection of facebook friends received through the fb callbacks in ReadFriendsList */
+	TArray< TSharedRef<FOnlineFriendFacebook> > CachedFriends;
+
+	/** Config based list of fields to use when querying friends list */
+	TArray<FString> FriendsFields;
+};
+
+typedef TSharedPtr<FOnlineFriendsFacebook, ESPMode::ThreadSafe> FOnlineFriendsFacebookPtr;
diff --git a/Engine/Plugins/Online/OnlineSubsystemFacebook/Source/Private/IOS/OnlineIdentityFacebook.cpp b/Engine/Plugins/Online/OnlineSubsystemFacebook/Source/Private/IOS/OnlineIdentityFacebook.cpp
new file mode 100644
--- /dev/null
+++ b/Engine/Plugins/Online/OnlineSubsystemFacebook/Source/Private/IOS/OnlineIdentityFacebook.cpp
@@ -0,0 +1,255 @@
+// Copyright 1998-2016 Epic Games, Inc. All Rights Reserved.
+
+
+// Module includes
+#include "OnlineSubsystemFacebookPrivatePCH.h"
+#include "OnlineIdentityFacebook.h"
+
+#import <FBSDKCoreKit/FBSDKCoreKit.h>
+#import <FBSDKLoginKit/FBSDKLoginKit.h>
+
+// Other UE4 includes
+#include "IOSAppDelegate.h"
+
+
+///////////////////////////////////////////////////////////////////////////////////////
+// FUserOnlineAccountFacebook implementation
+
+
+FUserOnlineAccountFacebook::FUserOnlineAccountFacebook(const FString& InUserId, const FString& InAuthTicket) 
+		: AuthTicket(InAuthTicket)
+		, UserId(new FUniqueNetIdString(InUserId))
+{
+	GConfig->GetString(TEXT("OnlineSubsystemFacebook.Login"), TEXT("AuthToken"), AuthTicket, GEngineIni);
+}
+
+TSharedRef<const FUniqueNetId> FUserOnlineAccountFacebook::GetUserId() const
+{
+	return UserId;
+}
+
+FString FUserOnlineAccountFacebook::GetRealName() const
+{
+	//@todo samz - implement
+	return FString();
+}
+
+FString FUserOnlineAccountFacebook::GetDisplayName(const FString& Platform) const
+{
+	//@todo samz - implement
+	return FString();
+}
+
+bool FUserOnlineAccountFacebook::GetUserAttribute(const FString& AttrName, FString& OutAttrValue) const
+{
+	return false;
+}
+
+bool FUserOnlineAccountFacebook::SetUserAttribute(const FString& AttrName, const FString& AttrValue)
+{
+	return false;
+}
+
+FString FUserOnlineAccountFacebook::GetAccessToken() const
+{
+	return AuthTicket;
+}
+
+bool FUserOnlineAccountFacebook::GetAuthAttribute(const FString& AttrName, FString& OutAttrValue) const
+{
+	return false;
+}
+
+///////////////////////////////////////////////////////////////////////////////////////
+// FOnlineIdentityFacebook implementation
+
+
+FOnlineIdentityFacebook::FOnlineIdentityFacebook()
+	: UserAccount( MakeShareable(new FUserOnlineAccountFacebook()) )
+	, LoginStatus( ELoginStatus::NotLoggedIn )
+{
+
+}
+
+
+TSharedPtr<FUserOnlineAccount> FOnlineIdentityFacebook::GetUserAccount(const FUniqueNetId& UserId) const
+{
+	return UserAccount;
+}
+
+
+TArray<TSharedPtr<FUserOnlineAccount> > FOnlineIdentityFacebook::GetAllUserAccounts() const
+{
+	TArray<TSharedPtr<FUserOnlineAccount> > Result;
+
+	Result.Add( UserAccount );
+
+	return Result;
+}
+
+
+TSharedPtr<const FUniqueNetId> FOnlineIdentityFacebook::GetUniquePlayerId(int32 LocalUserNum) const
+{
+	return UserAccount->GetUserId();
+}
+
+
+bool FOnlineIdentityFacebook::Login(int32 LocalUserNum, const FOnlineAccountCredentials& AccountCredentials)
+{
+	bool bTriggeredLogin = true;
+
+	dispatch_async(dispatch_get_main_queue(),^ 
+		{
+			FBSDKAccessToken *accessToken = [FBSDKAccessToken currentAccessToken];
+			if (accessToken == nil)
+			{
+				FBSDKLoginManager* loginManager = [[FBSDKLoginManager alloc] init];
+				[loginManager logInWithReadPermissions:nil
+					handler: ^(FBSDKLoginManagerLoginResult* result, NSError* error)
+					{
+						UE_LOG(LogOnline, Display, TEXT("[FBSDKLoginManager logInWithReadPermissions]"));
+						bool bSuccessfulLogin = false;
+				 
+						if(error)
+						{
+							UE_LOG(LogOnline, Display, TEXT("[FBSDKLoginManager logInWithReadPermissions = error[%d]]"), [error code]);
+						}
+						else if(result.isCancelled)
+						{
+							UE_LOG(LogOnline, Display, TEXT("[FBSDKLoginManager logInWithReadPermissions = cancelled"));
+						}
+						else
+						{
+							UE_LOG(LogOnline, Display, TEXT("[FBSDKLoginManager logInWithReadPermissions = true]"));
+				 
+							FString Token([result token].tokenString);
+							UserAccount->AuthTicket = Token;
+							GConfig->SetString(TEXT("OnlineSubsystemFacebook.Login"), TEXT("AuthToken"), *Token, GEngineIni);
+				 
+							const FString UserId([result token].userID);
+							UserAccount->UserId = MakeShareable(new FUniqueNetIdString(UserId));
+				 
+							bSuccessfulLogin = true;
+						}
+				 
+						LoginStatus = bSuccessfulLogin ? ELoginStatus::LoggedIn : ELoginStatus::NotLoggedIn;
+						FUniqueNetIdString TempId;
+						TriggerOnLoginCompleteDelegates(LocalUserNum, bSuccessfulLogin, (bSuccessfulLogin ? UserAccount->UserId.Get() : TempId), TEXT(""));
+				 
+						UE_LOG(LogOnline, Display, TEXT("Facebook login was successful? - %d"), bSuccessfulLogin);
+					}
+				];
+			}
+			else
+			{
+				LoginStatus = ELoginStatus::LoggedIn;
+				const FString UserId([accessToken userID]);
+				UserAccount->UserId = MakeShareable(new FUniqueNetIdString(UserId));
+
+				FString Token([accessToken tokenString]);
+				UserAccount->AuthTicket = Token;
+				GConfig->SetString(TEXT("OnlineSubsystemFacebook.Login"), TEXT("AuthToken"), *Token, GEngineIni);
+
+				TriggerOnLoginCompleteDelegates(LocalUserNum, true, UserAccount->UserId.Get(), TEXT(""));
+				
+				UE_LOG(LogOnline, Display, TEXT("Facebook login was successful? - Already had token!"));
+			}
+		}
+	);
+
+	return bTriggeredLogin;	
+}
+
+
+TSharedPtr<const FUniqueNetId> FOnlineIdentityFacebook::CreateUniquePlayerId(uint8* Bytes, int32 Size)
+{
+	if (Bytes != NULL && Size > 0)
+	{
+		FString StrId(Size, (TCHAR*)Bytes);
+		return MakeShareable(new FUniqueNetIdString(StrId));
+	}
+	return NULL;
+}
+
+
+TSharedPtr<const FUniqueNetId> FOnlineIdentityFacebook::CreateUniquePlayerId(const FString& Str)
+{
+	return MakeShareable(new FUniqueNetIdString(Str));
+}
+
+
+bool FOnlineIdentityFacebook::Logout(int32 LocalUserNum)
+{
+	//@todo samz - login status change delegate
+	dispatch_async(dispatch_get_main_queue(),^ 
+		{
+			if ([FBSDKAccessToken currentAccessToken])
+			{
+				FBSDKLoginManager *loginManager = [[FBSDKLoginManager alloc] init];
+				[loginManager logOut];
+				
+				LoginStatus = ELoginStatus::NotLoggedIn;
+			}
+		}
+	);
+
+	return true;
+}
+
+
+bool FOnlineIdentityFacebook::AutoLogin(int32 LocalUserNum)
+{
+	return false;
+}
+
+
+ELoginStatus::Type FOnlineIdentityFacebook::GetLoginStatus(int32 LocalUserNum) const
+{
+	return LoginStatus;
+}
+
+ELoginStatus::Type FOnlineIdentityFacebook::GetLoginStatus(const FUniqueNetId& UserId) const 
+{
+	return LoginStatus;
+}
+
+
+FString FOnlineIdentityFacebook::GetPlayerNickname(int32 LocalUserNum) const
+{
+	return TEXT("FacebookUser");
+}
+
+FString FOnlineIdentityFacebook::GetPlayerNickname(const FUniqueNetId& UserId) const 
+{
+	return TEXT("FacebookUser");
+}
+
+
+FString FOnlineIdentityFacebook::GetAuthToken(int32 LocalUserNum) const
+{
+	return UserAccount->GetAccessToken();
+}
+
+void FOnlineIdentityFacebook::GetUserPrivilege(const FUniqueNetId& UserId, EUserPrivileges::Type Privilege, const FOnGetUserPrivilegeCompleteDelegate& Delegate)
+{
+	Delegate.ExecuteIfBound(UserId, Privilege, (uint32)EPrivilegeResults::NoFailures);
+}
+
+FPlatformUserId FOnlineIdentityFacebook::GetPlatformUserIdFromUniqueNetId(const FUniqueNetId& UniqueNetId)
+{
+	for (int i = 0; i < MAX_LOCAL_PLAYERS; ++i)
+	{
+		auto CurrentUniqueId = GetUniquePlayerId(i);
+		if (CurrentUniqueId.IsValid() && (*CurrentUniqueId == UniqueNetId))
+		{
+			return i;
+		}
+	}
+
+	return PLATFORMUSERID_NONE;
+}
+
+FString FOnlineIdentityFacebook::GetAuthType() const
+{
+	return TEXT("");
+}
diff --git a/Engine/Plugins/Online/OnlineSubsystemFacebook/Source/Private/IOS/OnlineIdentityFacebook.h b/Engine/Plugins/Online/OnlineSubsystemFacebook/Source/Private/IOS/OnlineIdentityFacebook.h
new file mode 100644
--- /dev/null
+++ b/Engine/Plugins/Online/OnlineSubsystemFacebook/Source/Private/IOS/OnlineIdentityFacebook.h
@@ -0,0 +1,116 @@
+// Copyright 1998-2016 Epic Games, Inc. All Rights Reserved.
+
+
+#pragma once
+ 
+// Module includes
+#include "OnlineIdentityInterface.h"
+#include "OnlineSubsystemFacebookPackage.h"
+
+/**
+ * Facebook implementation of the online account information we may want
+ */
+class FUserOnlineAccountFacebook : public FUserOnlineAccount
+{
+
+public:
+
+	// FOnlineUser
+	
+	virtual TSharedRef<const FUniqueNetId> GetUserId() const override;
+	virtual FString GetRealName() const override;
+	virtual FString GetDisplayName(const FString& Platform = FString()) const override;
+	virtual bool GetUserAttribute(const FString& AttrName, FString& OutAttrValue) const override;
+	virtual bool SetUserAttribute(const FString& AttrName, const FString& AttrValue) override;
+
+	// FUserOnlineAccount
+
+	virtual FString GetAccessToken() const override;
+	virtual bool GetAuthAttribute(const FString& AttrName, FString& OutAttrValue) const override;
+
+public:
+
+	/**
+	 * Constructor
+	 *
+	 * @param	InUserId		The user Id which should be assigned for this account.
+	 * @param	InAuthTicket	The token which this account is using.
+	 */
+	FUserOnlineAccountFacebook(const FString& InUserId=TEXT(""), const FString& InAuthTicket=TEXT("")) ;
+
+	/**
+	 * Destructor
+	 */
+	virtual ~FUserOnlineAccountFacebook(){}
+
+
+private:
+
+	/** The token obtained for this session */
+	FString AuthTicket;
+
+	/** UID of the logged in user */
+	TSharedRef<const FUniqueNetId> UserId;
+
+	/** Allow the fb identity to fill in our private members from it's callbacks */
+	friend class FOnlineIdentityFacebook;
+};
+
+
+/**
+ * Facebook service implementation of the online identity interface
+ */
+class FOnlineIdentityFacebook : public IOnlineIdentity
+{
+
+public:
+
+	//~ Begin IOnlineIdentity Interface	
+	virtual bool Login(int32 LocalUserNum, const FOnlineAccountCredentials& AccountCredentials) override;
+	virtual bool Logout(int32 LocalUserNum) override;
+	virtual bool AutoLogin(int32 LocalUserNum) override;
+	virtual TSharedPtr<FUserOnlineAccount> GetUserAccount(const FUniqueNetId& UserId) const override;
+	virtual TArray<TSharedPtr<FUserOnlineAccount> > GetAllUserAccounts() const override;
+	virtual TSharedPtr<const FUniqueNetId> GetUniquePlayerId(int32 LocalUserNum) const override;
+	virtual TSharedPtr<const FUniqueNetId> CreateUniquePlayerId(uint8* Bytes, int32 Size) override;
+	virtual TSharedPtr<const FUniqueNetId> CreateUniquePlayerId(const FString& Str) override;
+	virtual ELoginStatus::Type GetLoginStatus(int32 LocalUserNum) const override;
+	virtual ELoginStatus::Type GetLoginStatus(const FUniqueNetId& UserId) const override;
+	virtual FString GetPlayerNickname(int32 LocalUserNum) const override;
+	virtual FString GetPlayerNickname(const FUniqueNetId& UserId) const override;
+	virtual FString GetAuthToken(int32 LocalUserNum) const override;
+	virtual void GetUserPrivilege(const FUniqueNetId& UserId, EUserPrivileges::Type Privilege, const FOnGetUserPrivilegeCompleteDelegate& Delegate) override;
+	virtual FPlatformUserId GetPlatformUserIdFromUniqueNetId(const FUniqueNetId& UniqueNetId) override;
+	virtual FString GetAuthType() const override;
+	//~ End IOnlineIdentity Interface
+
+
+public:
+
+	/**
+	 * Default constructor
+	 */
+	FOnlineIdentityFacebook();
+
+
+	/**
+	 * Destructor
+	 */
+	virtual ~FOnlineIdentityFacebook()
+	{
+	}
+
+
+private:
+
+	/** The user account details associated with this identity */
+	TSharedRef< FUserOnlineAccountFacebook > UserAccount;
+
+	/** The current state of our login */
+	ELoginStatus::Type LoginStatus;
+
+	/** Username of the user assuming this identity */
+	FString Name;
+};
+
+typedef TSharedPtr<FOnlineIdentityFacebook, ESPMode::ThreadSafe> FOnlineIdentityFacebookPtr;
\ No newline at end of file
