Copying //UE4/Ocean-Staging to //UE4/Dev-Main (Source: //Ocean/Main @ 3140098)

==========================
MAJOR FEATURES + CHANGES
==========================

Change 3138246 on 2016/09/23 by [REDACTED_PERSON]

	Integrating codesign fix into //Ocean/Main

Change 3129182 on 2016/09/16 by [REDACTED_PERSON]

	OGS game plugin directory integration  (Ocean has own copy of  OSS)

	Merging

	//Orion/Main/Engine/...

	to //Ocean/Main/Engine/...

Change 3129181 on 2016/09/16 by [REDACTED_PERSON]

	OGS game plugin directory integration  (Ocean has own copy of  OSS)
	Merging

	//UE4/Main/Engine/Source/Runtime/Online/XMPP/...

	to //Ocean/Main/Engine/Source/Runtime/Online/XMPP/...

Change 3129178 on 2016/09/16 by [REDACTED_PERSON]

	#UE4 - changes related to OGF/OSS pluginization effort
	OGS game plugin directory integration  (Ocean has own copy of  OSS)

Change 3127280 on 2016/09/15 by [REDACTED_PERSON]

	#ocn - Turn off audio and distortion and decrease shadow depth buffer size for iOS devices with only 1Gb of memory.

Change 3118202 on 2016/09/08 by [REDACTED_PERSON]

	Return UE_GAME from IsGameOnly on iOS.

Change 3116888 on 2016/09/07 by [REDACTED_PERSON]

	Unload bulk audio data if audio is disabled.

Change 3099939 on 2016/08/24 by [REDACTED_PERSON]

	Back out changelist 3099720 - Was crashing the editor on load.

Change 3099720 on 2016/08/24 by [REDACTED_PERSON]

	Don't load bulk data for sound files when audio is disabled.

Change 3099714 on 2016/08/24 by [REDACTED_PERSON]

	Integrate fix for audio halt from 4.13.

Change 3096950 on 2016/08/22 by [REDACTED_PERSON]

	Merging - Spot merged CL# 3086032 to fix an editor crash from last merge.

Change 3095606 on 2016/08/19 by [REDACTED_PERSON]

	Perf: Updated Blueprint search indexing.

Change 3094810 on 2016/08/19 by [REDACTED_PERSON]

	Merging //UE4/Ocean-Staging->//Ocean/Main

Change 3091147 on 2016/08/16 by [REDACTED_PERSON]

	#core - Added a simple function in FStringAssetReference and AssetPtr for GetAssetName

Change 3084075 on 2016/08/10 by [REDACTED_PERSON]

	Fixed a mobile shadow rendering bug that was causing a crash in the Ocean onboarding level.

Change 3077991 on 2016/08/04 by [REDACTED_PERSON]

	Testing Orion fix for font crash.

Change 3077854 on 2016/08/04 by [REDACTED_PERSON]

	Turn off the platform file cached wrapper for iOS and Android.  Allows Ocean to load on iPhone 6s without running out of memory.

Change 3075681 on 2016/08/03 by [REDACTED_PERSON]

	Change iPhonePackager to 64-bit so it doesn't run out of memory packaging Ocean.

Change 3074498 on 2016/08/02 by [REDACTED_PERSON]

	Hotfix from CL 3031932 for UE-32614 - Blueprint pin value data missing from tooltip

Change 3073956 on 2016/08/02 by [REDACTED_PERSON]

	Hotfix from CL 3073388 for UE-33821: Async Loading issue in Ocean.

Change 3072524 on 2016/08/01 by [REDACTED_PERSON]

	#ocn, #mcp: Updated re-init behavior for online groups client implementation.

Change 3072392 on 2016/08/01 by [REDACTED_PERSON]

	#ocn, #mcp: Changed ini to add groups api namespace configuration information.  Removed setting the namespace in code.

Change 3067671 on 2016/07/27 by [REDACTED_PERSON]

	Disabling oculus related plugins and reverting the temporary workaround for the ocean win64 non-editor build compile error. The error is due to the need to compile with VS 2013 until the VS 2015 redistributable bugs are sorted out. Excluding the oculus plugins is another (and less error-prone) way to workaround the issue.

Change 3065390 on 2016/07/26 by [REDACTED_PERSON]

	Hotfix from CL 3065307 for #jira UE-33457: Importing a font crashes the editor,

Change 3062603 on 2016/07/23 by [REDACTED_PERSON]

	Merging CL 3041439 to fix Mac non-unity issues.

Change 3059851 on 2016/07/21 by [REDACTED_PERSON]

	Hotfix from CL 3059493 for removing the hack forcing CrashReportClient to be built with Visual Studio 2013.

Change 3057815 on 2016/07/20 by [REDACTED_PERSON]

	Hotfix from CL 3056841 for UE-33362 - Unable to move camera around in viewport after selecting eject in Ocean

Change 3053947 on 2016/07/18 by [REDACTED_PERSON]

	Temporarily removing the VS 2013 hack (the engine team is currently discussing what the permanent solution should be, but for now, this was recommended in our branch)

Change 3053925 on 2016/07/18 by [REDACTED_PERSON]

	Fixing an issue with UGS attempting to build CrashReportClient with VS 2013

Change 3052480 on 2016/07/15 by [REDACTED_PERSON]

	Temporary fix for build machine errors until we can figure out why this is causing issues

Change 3051809 on 2016/07/15 by [REDACTED_PERSON]

	Fixing a mac build warning

Change 3050835 on 2016/07/14 by [REDACTED_PERSON]

	Merging using //Ocean/Main->//UE4/Ocean-Staging

Change 3050485 on 2016/07/14 by [REDACTED_PERSON]

	#ocn, #mcp: Added missing clan call - cancel pending join request.

Change 3050239 on 2016/07/14 by [REDACTED_PERSON]

	#ocn, #mcp: Added support for querying 'headcount' and 'membership' configuration settings for the Groups service.

Change 3046132 on 2016/07/12 by [REDACTED_PERSON]

	#UE4, #OnlineSubSystem - Hotfix in the ExpirationDate field early so we can update the OGF plugin

Change 3041256 on 2016/07/07 by [REDACTED_PERSON]

	Merge latest BuildGraph changes from //UE4/Ocean-Staging.

diff --git a/Engine/Plugins/Online/OnlineSubsystem/Source/Public/Interfaces/OnlineGroupsInterface.h b/Engine/Plugins/Online/OnlineSubsystem/Source/Public/Interfaces/OnlineGroupsInterface.h
--- a/Engine/Plugins/Online/OnlineSubsystem/Source/Public/Interfaces/OnlineGroupsInterface.h
+++ b/Engine/Plugins/Online/OnlineSubsystem/Source/Public/Interfaces/OnlineGroupsInterface.h
@@ -194,6 +194,23 @@ struct FInvitationEntry
 	FText Name;
 };
 
+template <typename EntryType>
+struct TGroupConfigEntry
+{
+	TSharedPtr<const FUniqueNetId> GetId() const { return AccountId; }
+
+	/** Context AccountId. */
+	TSharedPtr<const FUniqueNetId> AccountId;
+
+	/** The max number of members in any group. */
+	FString Key;
+
+	/** The max number of groups in which a user may be a member. */
+	EntryType Value;
+};
+
+typedef TGroupConfigEntry<int32> FGroupConfigEntryInt;
+
 template <typename EntryType>
 struct IGroupUserCollection
 {
@@ -238,6 +255,11 @@ typedef IGroupUserCollection<FApplicationEntry> IApplications;
 */
 typedef IGroupUserCollection<FInvitationEntry> IInvitations;
 
+/**
+* Configuration key->values (int)
+*/
+typedef TMap<FString, TSharedPtr<const FGroupConfigEntryInt>> IGroupConfigInt;
+
 /**
  * Group search options
  */
@@ -409,6 +431,16 @@ public: // callable by all users
 	 */
 	virtual void LeaveGroup(const FUniqueNetId& ContextUserId, const FUniqueNetId& GroupId, const FOnGroupsRequestCompleted& OnCompleted) = 0;
 
+	/**
+	* Cancel pending request to join the given group.
+	*
+	* @param ContextUserId The ID of the user whose credentials are being used to make this call
+	* @param GroupId The group's globally unique ID
+	* @param OnCompleted This callback is invoked after contacting the server. It is guaranteed to occur
+	*        (regardless of success/fail) and will not be called before this function returns.
+	*/
+	virtual void CancelRequest(const FUniqueNetId& ContextUserId, const FUniqueNetId& GroupId, const FOnGroupsRequestCompleted& OnCompleted) = 0;
+
 	/**
 	 * Accept a pending invite to join a group. This is done on behalf of the context
 	 * user (provided when you request this IOnlineGroups interface) and also updates the their cached membership
@@ -744,6 +776,24 @@ public: // can be called by group admins
 	*/
 	virtual void QueryIncomingApplications(const FUniqueNetId& ContextUserId, const FUniqueNetId& UserId, const FOnGroupsRequestCompleted& OnCompleted) = 0;
 
+public: // configuration queries
+
+	/**
+	* Queries the system configuration for system-wide group max membership headcount.
+	*/
+	virtual void QueryConfigHeadcount(const FUniqueNetId& ContextUserId, const FOnGroupsRequestCompleted& OnCompleted) = 0;
+
+	/**
+	* Queries the system configuration for system-wide user max membership count.
+	*/
+	virtual void QueryConfigMembership(const FUniqueNetId& ContextUserId, const FOnGroupsRequestCompleted& OnCompleted) = 0;
+
+	/**
+	* Get the result of a previous configuration query.
+	* Only the member queried for will be valid.
+	*/
+	virtual TSharedPtr<const FGroupConfigEntryInt> GetCachedConfigInt(const FString& Key) = 0;
+
 public: // can be called by group owner only
 
 	/**
@@ -771,6 +821,7 @@ public: // can be called by group owner only
 
 protected:
 	IOnlineGroups() {}
+
 public:
 	virtual ~IOnlineGroups() {};
 
diff --git a/Engine/Source/Runtime/Core/Public/IOS/IOSPlatformProperties.h b/Engine/Source/Runtime/Core/Public/IOS/IOSPlatformProperties.h
--- a/Engine/Source/Runtime/Core/Public/IOS/IOSPlatformProperties.h
+++ b/Engine/Source/Runtime/Core/Public/IOS/IOSPlatformProperties.h
@@ -33,6 +33,11 @@ struct FIOSPlatformProperties
 		return "IOS";
 	}
 
+	static FORCEINLINE bool IsGameOnly()
+	{
+		return UE_GAME;
+	}
+	
 	static FORCEINLINE bool RequiresCookedData( )
 	{
 		return true;
diff --git a/Engine/Source/Runtime/CoreUObject/Public/Misc/StringAssetReference.h b/Engine/Source/Runtime/CoreUObject/Public/Misc/StringAssetReference.h
--- a/Engine/Source/Runtime/CoreUObject/Public/Misc/StringAssetReference.h
+++ b/Engine/Source/Runtime/CoreUObject/Public/Misc/StringAssetReference.h
@@ -50,6 +50,13 @@ struct COREUOBJECT_API FStringAssetReference
 		return PackageName;
 	}
 
+	const FString GetAssetName() const
+	{
+		FString AssetName;
+		ToString().Split(TEXT("."), nullptr, &AssetName, ESearchCase::CaseSensitive, ESearchDir::FromEnd);
+		return AssetName;
+	}
+
 	/**
 	 * Sets asset path of this reference.
 	 *
diff --git a/Engine/Source/Runtime/CoreUObject/Public/UObject/AssetPtr.h b/Engine/Source/Runtime/CoreUObject/Public/UObject/AssetPtr.h
--- a/Engine/Source/Runtime/CoreUObject/Public/UObject/AssetPtr.h
+++ b/Engine/Source/Runtime/CoreUObject/Public/UObject/AssetPtr.h
@@ -279,6 +279,11 @@ public:
 		return ToStringReference().GetLongPackageName();
 	}
 
+	FORCEINLINE FString GetAssetName() const
+	{
+		return ToStringReference().GetAssetName();
+	}
+
 	/**  
 	 * Dereference lazy pointer to see if it points somewhere valid.
 	 */
diff --git a/Engine/Source/Runtime/Engine/Classes/Engine/Engine.h b/Engine/Source/Runtime/Engine/Classes/Engine/Engine.h
--- a/Engine/Source/Runtime/Engine/Classes/Engine/Engine.h
+++ b/Engine/Source/Runtime/Engine/Classes/Engine/Engine.h
@@ -1354,11 +1354,13 @@ public:
 	UPROPERTY(globalconfig)
 	uint32 bEnableVisualLogRecordingOnStart;
 
-private:
+protected:
 
 	/** Whether the engine should be playing sounds.  If false at initialization time the AudioDevice will not be created */
 	uint32 bUseSound:1;
 
+private:
+	
 	/** Semaphore to control screen saver inhibitor thread access. */
 	UPROPERTY(transient)
 	int32 ScreenSaverInhibitorSemaphore;
diff --git a/Engine/Source/Runtime/Engine/Private/AudioDevice.cpp b/Engine/Source/Runtime/Engine/Private/AudioDevice.cpp
--- a/Engine/Source/Runtime/Engine/Private/AudioDevice.cpp
+++ b/Engine/Source/Runtime/Engine/Private/AudioDevice.cpp
@@ -232,7 +232,7 @@ float FAudioDevice::GetLowPassFilterResonance() const
 void FAudioDevice::PrecacheStartupSounds()
 {
 	// Iterate over all already loaded sounds and precache them. This relies on Super::Init in derived classes to be called last.
-	if (!GIsEditor)
+	if (!GIsEditor && GEngine && GEngine->UseSound() )
 	{
 		for (TObjectIterator<USoundWave> It; It; ++It)
 		{
diff --git a/Engine/Source/Runtime/Renderer/Private/DistortionRendering.cpp b/Engine/Source/Runtime/Renderer/Private/DistortionRendering.cpp
--- a/Engine/Source/Runtime/Renderer/Private/DistortionRendering.cpp
+++ b/Engine/Source/Runtime/Renderer/Private/DistortionRendering.cpp
@@ -15,6 +15,12 @@
 
 DECLARE_FLOAT_COUNTER_STAT(TEXT("Distortion"), Stat_GPU_Distortion, STATGROUP_GPU);
 
+static TAutoConsoleVariable<int32> CVarDisableDistortion(
+														 TEXT("r.DisableDistortion"),
+														 0,
+														 TEXT("Prevents distortion effects from rendering.  Saves a full-screen framebuffer's worth of memory."),
+														 ECVF_Default);
+
 /**
 * A pixel shader for rendering the full screen refraction pass
 */
@@ -1041,8 +1047,11 @@ void FSceneRenderer::RenderDistortionES2(FRHICommandListImmediate& RHICmdList)
 			break;
 		}
 	}
-			
-	if (bRender)
+	
+	static IConsoleVariable* CVar = IConsoleManager::Get().FindConsoleVariable(TEXT("r.DisableDistortion"));
+	int32 DisableDistortion = CVar->GetInt();
+	
+	if (bRender && !DisableDistortion)
 	{
 		// Apply distortion
 		SCOPED_DRAW_EVENT(RHICmdList, Distortion);
