Fix the CIS warning (OnlineSubsystemNone)

diff --git a/Engine/Plugins/Online/OnlineSubsystem/Source/Private/OnlineSubsystemModule.cpp b/Engine/Plugins/Online/OnlineSubsystem/Source/Private/OnlineSubsystemModule.cpp
--- a/Engine/Plugins/Online/OnlineSubsystem/Source/Private/OnlineSubsystemModule.cpp
+++ b/Engine/Plugins/Online/OnlineSubsystem/Source/Private/OnlineSubsystemModule.cpp
@@ -116,7 +116,11 @@ void FOnlineSubsystemModule::LoadDefaultSubsystem()
 	// if the above failed, then attempt the platform's default
 	if (!bHasLoadedModule)
 	{
-		bHasLoadedModule = TryLoadSubsystemAndSetDefault(FPlatformMisc::GetDefaultOnlineSubsystemName());
+		FName PlatformOSSName = FPlatformMisc::GetDefaultOnlineSubsystemName();
+		if (PlatformOSSName != NAME_None)
+		{
+			bHasLoadedModule = TryLoadSubsystemAndSetDefault(PlatformOSSName);
+		}
 	}
 
 	// if all else fails, attempt to load Null
