Copying //UE4/Dev-Online to //UE4/Dev-Main (Source: //UE4/Dev-Online [at] 4192069)

============================
  MAJOR FEATURES & CHANGES
============================

Change 4052994 by [REDACTED_PERSON]

	#IOS - Fix default OSS being wrong

Change 4053003 by [REDACTED_PERSON]

	#UE4 - Remove Pragma Once from all cpp files
	- In spirit of github request #3940 that was already fixed

Change 4059536 by [REDACTED_PERSON]

	#ShooterGame - Do not use localized names for game-types

Change 4179456 by [REDACTED_PERSON]

	#OnlineSubsystem - Add Tournament OSS Interface

diff --git a/Engine/Config/IOS/IOSEngine.ini b/Engine/Config/IOS/IOSEngine.ini
--- a/Engine/Config/IOS/IOSEngine.ini
+++ b/Engine/Config/IOS/IOSEngine.ini
@@ -25,7 +25,7 @@ SmallerMemoryBucket_MinGB=3
 SmallestMemoryBucket_MinGB=2
 
 [OnlineSubsystem]
-DefaultPlatformService=MCP
+DefaultPlatformService=IOS
 
 [Advertising]
 DefaultProviderName=IOSAdvertising
diff --git a/Engine/Plugins/Online/Android/OnlineSubsystemGameCircle/Source/OnlineSubsystemGameCircle/Public/OnlineSubsystemGameCircle.h b/Engine/Plugins/Online/Android/OnlineSubsystemGameCircle/Source/OnlineSubsystemGameCircle/Public/OnlineSubsystemGameCircle.h
--- a/Engine/Plugins/Online/Android/OnlineSubsystemGameCircle/Source/OnlineSubsystemGameCircle/Public/OnlineSubsystemGameCircle.h
+++ b/Engine/Plugins/Online/Android/OnlineSubsystemGameCircle/Source/OnlineSubsystemGameCircle/Public/OnlineSubsystemGameCircle.h
@@ -51,6 +51,7 @@ public:
 	virtual IOnlinePresencePtr GetPresenceInterface() const override { return nullptr; }
 	virtual IOnlineChatPtr GetChatInterface() const override { return nullptr; }
 	virtual IOnlineTurnBasedPtr GetTurnBasedInterface() const override { return nullptr; }
+	virtual IOnlineTournamentPtr GetTournamentInterface() const override { return nullptr; }
 
 	virtual class UObject* GetNamedInterface(FName InterfaceName) override { return nullptr; }
 	virtual void SetNamedInterface(FName InterfaceName, class UObject* NewInterface) override {}
diff --git a/Engine/Plugins/Online/Android/OnlineSubsystemGooglePlay/Source/Public/OnlineSubsystemGooglePlay.h b/Engine/Plugins/Online/Android/OnlineSubsystemGooglePlay/Source/Public/OnlineSubsystemGooglePlay.h
--- a/Engine/Plugins/Online/Android/OnlineSubsystemGooglePlay/Source/Public/OnlineSubsystemGooglePlay.h
+++ b/Engine/Plugins/Online/Android/OnlineSubsystemGooglePlay/Source/Public/OnlineSubsystemGooglePlay.h
@@ -67,6 +67,7 @@ public:
 	virtual IOnlinePresencePtr GetPresenceInterface() const override { return nullptr; }
 	virtual IOnlineChatPtr GetChatInterface() const override { return nullptr; }
 	virtual IOnlineTurnBasedPtr GetTurnBasedInterface() const override { return nullptr; }
+	virtual IOnlineTournamentPtr GetTournamentInterface() const override { return nullptr; }
 
 	virtual class UObject* GetNamedInterface(FName InterfaceName) override { return nullptr; }
 	virtual void SetNamedInterface(FName InterfaceName, class UObject* NewInterface) override {}
diff --git a/Engine/Plugins/Online/IOS/OnlineSubsystemIOS/Source/Private/OnlineSubsystemIOS.cpp b/Engine/Plugins/Online/IOS/OnlineSubsystemIOS/Source/Private/OnlineSubsystemIOS.cpp
--- a/Engine/Plugins/Online/IOS/OnlineSubsystemIOS/Source/Private/OnlineSubsystemIOS.cpp
+++ b/Engine/Plugins/Online/IOS/OnlineSubsystemIOS/Source/Private/OnlineSubsystemIOS.cpp
@@ -136,7 +136,12 @@ IOnlineChatPtr FOnlineSubsystemIOS::GetChatInterface() const
 
 IOnlineTurnBasedPtr FOnlineSubsystemIOS::GetTurnBasedInterface() const
 {
-    return TurnBasedInterface;
+	return TurnBasedInterface;
+}
+
+IOnlineTournamentPtr FOnlineSubsystemIOS::GetTournamentInterface() const
+{
+	return nullptr;
 }
 
 bool FOnlineSubsystemIOS::Init() 
diff --git a/Engine/Plugins/Online/IOS/OnlineSubsystemIOS/Source/Public/OnlineSubsystemIOS.h b/Engine/Plugins/Online/IOS/OnlineSubsystemIOS/Source/Public/OnlineSubsystemIOS.h
--- a/Engine/Plugins/Online/IOS/OnlineSubsystemIOS/Source/Public/OnlineSubsystemIOS.h
+++ b/Engine/Plugins/Online/IOS/OnlineSubsystemIOS/Source/Public/OnlineSubsystemIOS.h
@@ -43,9 +43,10 @@ public:
 	virtual IOnlineSharingPtr GetSharingInterface() const override;
 	virtual IOnlineUserPtr GetUserInterface() const override;
 	virtual IOnlineMessagePtr GetMessageInterface() const override;
-    virtual IOnlinePresencePtr GetPresenceInterface() const override;
+	virtual IOnlinePresencePtr GetPresenceInterface() const override;
 	virtual IOnlineChatPtr GetChatInterface() const override;
-    virtual IOnlineTurnBasedPtr GetTurnBasedInterface() const override;
+	virtual IOnlineTurnBasedPtr GetTurnBasedInterface() const override;
+	virtual IOnlineTournamentPtr GetTournamentInterface() const override;
 	virtual bool Init() override;
 	virtual bool Shutdown() override;
 	virtual FString GetAppId() const override;
diff --git a/Engine/Plugins/Online/OnlineSubsystem/Source/Private/OnlineSubsystem.cpp b/Engine/Plugins/Online/OnlineSubsystem/Source/Private/OnlineSubsystem.cpp
--- a/Engine/Plugins/Online/OnlineSubsystem/Source/Private/OnlineSubsystem.cpp
+++ b/Engine/Plugins/Online/OnlineSubsystem/Source/Private/OnlineSubsystem.cpp
@@ -19,6 +19,7 @@
 #include "Interfaces/OnlinePurchaseInterface.h"
 #include "Interfaces/OnlineSharingInterface.h"
 #include "Interfaces/OnlineLeaderboardInterface.h"
+#include "Interfaces/OnlineTournamentInterface.h"
 
 DEFINE_LOG_CATEGORY(LogOnline);
 DEFINE_LOG_CATEGORY(LogOnlineGame);
@@ -41,7 +42,7 @@ DEFINE_LOG_CATEGORY(LogOnlineSharing);
 DEFINE_LOG_CATEGORY(LogOnlineStore);
 DEFINE_LOG_CATEGORY(LogOnlineStoreV2);
 DEFINE_LOG_CATEGORY(LogOnlinePurchase);
-
+DEFINE_LOG_CATEGORY(LogOnlineTournament);
 
 #if STATS
 ONLINESUBSYSTEM_API DEFINE_STAT(STAT_Online_Async);
diff --git a/Engine/Plugins/Online/OnlineSubsystem/Source/Public/Interfaces/OnlineTournamentInterface.h b/Engine/Plugins/Online/OnlineSubsystem/Source/Public/Interfaces/OnlineTournamentInterface.h
new file mode 100644
--- /dev/null
+++ b/Engine/Plugins/Online/OnlineSubsystem/Source/Public/Interfaces/OnlineTournamentInterface.h
@@ -0,0 +1,816 @@
+// Copyright 1998-2018 Epic Games, Inc. All Rights Reserved.
+
+#pragma once
+
+#include "CoreMinimal.h"
+#include "OnlineSubsystemTypes.h"
+#include "OnlineKeyValuePair.h"
+
+ONLINESUBSYSTEM_API DECLARE_LOG_CATEGORY_EXTERN(LogOnlineTournament, Display, All);
+
+#define UE_LOG_ONLINE_TOURNAMENT(Verbosity, Format, ...) \
+{ \
+	UE_LOG(LogOnlineTournament, Verbosity, TEXT("%s%s"), ONLINE_LOG_PREFIX, *FString::Printf(Format, ##__VA_ARGS__)); \
+}
+
+#define UE_CLOG_ONLINE_TOURNAMENT(Conditional, Verbosity, Format, ...) \
+{ \
+	UE_CLOG(Conditional, LogOnlineTournament, Verbosity, TEXT("%s%s"), ONLINE_LOG_PREFIX, *FString::Printf(Format, ##__VA_ARGS__)); \
+}
+
+class FUniqueNetId;
+struct FOnlineError;
+
+/** UniqueNetId of a tournament */
+using FOnlineTournamentId = FUniqueNetId;
+
+/** UniqueNetId of a match in a tournament */
+using FOnlineTournamentMatchId = FUniqueNetId;
+
+/** UniqueNetId of a participant (Player or team) in a tournament */
+using FOnlineTournamentParticipantId = FUniqueNetId;
+
+/** UniqueNetId of a team in a tournament */
+using FOnlineTournamentTeamId = FOnlineTournamentParticipantId;
+
+/**
+ * What format the tournament is being run as
+ */
+enum class EOnlineTournamentFormat : uint8
+{
+	/** The tournament is being run in the single elimination format */
+	SingleElimination,
+	/** The tournament is being run in the double elimination format */
+	DoubleElimination,
+	/** The tournament is being run in the swiss format */
+	Swiss,
+	/** The tournament is being run in the round-robin format */
+	RoundRobin,
+	/** The tournament is being run in a custom format */
+	Custom
+};
+
+inline void LexFromString(TOptional<EOnlineTournamentFormat>& Format, const TCHAR* const String)
+{
+	if (FCString::Stricmp(String, TEXT("SingleElimination")) == 0)
+	{
+		Format = EOnlineTournamentFormat::SingleElimination;
+	}
+	else if (FCString::Stricmp(String, TEXT("DoubleElimination")) == 0)
+	{
+		Format = EOnlineTournamentFormat::DoubleElimination;
+	}
+	else if (FCString::Stricmp(String, TEXT("Swiss")) == 0)
+	{
+		Format = EOnlineTournamentFormat::Swiss;
+	}
+	else if (FCString::Stricmp(String, TEXT("RoundRobin")) == 0)
+	{
+		Format = EOnlineTournamentFormat::RoundRobin;
+	}
+	else if (FCString::Stricmp(String, TEXT("Custom")) == 0)
+	{
+		Format = EOnlineTournamentFormat::Custom;
+	}
+	else
+	{
+		Format = TOptional<EOnlineTournamentFormat>();
+	}
+}
+
+inline FString LexToString(const EOnlineTournamentFormat Format)
+{
+	switch (Format)
+	{
+	case EOnlineTournamentFormat::SingleElimination:
+		return TEXT("SingleElimination");
+	case EOnlineTournamentFormat::DoubleElimination:
+		return TEXT("DoubleElimination");
+	case EOnlineTournamentFormat::Swiss:
+		return TEXT("Swiss");
+	case EOnlineTournamentFormat::RoundRobin:
+		return TEXT("RoundRobin");
+	case EOnlineTournamentFormat::Custom:
+		return TEXT("Custom");
+	}
+
+	checkNoEntry();
+	return FString();
+}
+
+/**
+ * What state the tournament is currently in
+ */
+enum class EOnlineTournamentState : uint8
+{
+	/** The tournament has been created, but participants may not be registered yet */
+	Created,
+	/** The tournament is now open for registration */
+	OpenRegistration,
+	/** The tournament registration has now closed, but the tournament has not started yet */
+	ClosedRegistration,
+	/** The tournament is now in progress */
+	InProgress,
+	/** The tournament has now finished and all results are finalized */
+	Finished,
+	/** The tournament was cancelled */
+	Cancelled
+};
+
+inline void LexFromString(TOptional<EOnlineTournamentState>& State, const TCHAR* const String)
+{
+	if (FCString::Stricmp(String, TEXT("Created")) == 0)
+	{
+		State = EOnlineTournamentState::Created;
+	}
+	else if (FCString::Stricmp(String, TEXT("OpenRegistration")) == 0)
+	{
+		State = EOnlineTournamentState::OpenRegistration;
+	}
+	else if (FCString::Stricmp(String, TEXT("ClosedRegistration")) == 0)
+	{
+		State = EOnlineTournamentState::ClosedRegistration;
+	}
+	else if (FCString::Stricmp(String, TEXT("InProgress")) == 0)
+	{
+		State = EOnlineTournamentState::InProgress;
+	}
+	else if (FCString::Stricmp(String, TEXT("Finished")) == 0)
+	{
+		State = EOnlineTournamentState::Finished;
+	}
+	else if (FCString::Stricmp(String, TEXT("Cancelled")) == 0)
+	{
+		State = EOnlineTournamentState::Cancelled;
+	}
+	else
+	{
+		State = TOptional<EOnlineTournamentState>();
+	}
+}
+
+inline FString LexToString(const EOnlineTournamentState State)
+{
+	switch (State)
+	{
+	case EOnlineTournamentState::Created:
+		return TEXT("Created");
+	case EOnlineTournamentState::OpenRegistration:
+		return TEXT("OpenRegistration");
+	case EOnlineTournamentState::ClosedRegistration:
+		return TEXT("ClosedRegistration");
+	case EOnlineTournamentState::InProgress:
+		return TEXT("InProgress");
+	case EOnlineTournamentState::Finished:
+		return TEXT("Finished");
+	case EOnlineTournamentState::Cancelled:
+		return TEXT("Cancelled");
+	}
+
+	checkNoEntry();
+	return FString();
+}
+
+/**
+ * What participant format does this tournament support?
+ */
+enum class EOnlineTournamentParticipantType : uint8
+{
+	/** The tournament has individual players facing each other */
+	Individual,
+	/** The tournament has teams facing other teams */
+	Team
+};
+
+inline void LexFromString(TOptional<EOnlineTournamentParticipantType>& State, const TCHAR* const String)
+{
+	if (FCString::Stricmp(String, TEXT("Individual")) == 0)
+	{
+		State = EOnlineTournamentParticipantType::Individual;
+	}
+	else if (FCString::Stricmp(String, TEXT("Team")) == 0)
+	{
+		State = EOnlineTournamentParticipantType::Team;
+	}
+	else
+	{
+		State = TOptional<EOnlineTournamentParticipantType>();
+	}
+}
+
+inline FString LexToString(const EOnlineTournamentParticipantType ParticipantType)
+{
+	switch (ParticipantType)
+	{
+	case EOnlineTournamentParticipantType::Individual:
+		return TEXT("Individual");
+	case EOnlineTournamentParticipantType::Team:
+		return TEXT("Team");
+	}
+
+	checkNoEntry();
+	return FString();
+}
+
+/**
+ * Filters to use when querying for tournament information.
+ *
+ * Some of these fields may be required, depending on the backing online system.  Some fields
+ * may not be specified if other fields are specified, depending on the backing online system.
+ */
+struct FOnlineTournamentQueryFilter
+{
+public:
+	/** What direction to sort these results by (useful when specifying limits and offsetts)*/
+	enum class EOnlineTournamentSortDirection : uint8
+	{
+		/** Results will be sorted in Ascending order */
+		Ascending,
+		/** Results will be sorted in Descending order */
+		Descending
+	};
+
+	/** Filter tournament information that does not match this participant type */
+	TOptional<EOnlineTournamentParticipantType> ParticipantType;
+	/** Filter tournament information that does not match this tournament format */
+	TOptional<EOnlineTournamentFormat> Format;
+	/** Only include tournament information that includes this team (on team tournaments) */
+	TOptional<TSharedRef<const FOnlineTournamentTeamId>> TeamId;
+	/** Only include tournament information that includes this player */
+	TOptional<TSharedRef<const FUniqueNetId>> PlayerId;
+	/** Limit the results to this many entries */
+	TOptional<uint32> Limit;
+	/** Start the results this many entries in */
+	TOptional<uint32> Offset;
+	/** Sort the results in this direction */
+	TOptional<EOnlineTournamentSortDirection> SortDirection;
+};
+
+/**
+ * Details about a participant and their current score
+ */
+struct FOnlineTournamentScore
+{
+public:
+	virtual ~FOnlineTournamentScore() = default;
+
+	FOnlineTournamentScore(const TSharedRef<const FOnlineTournamentParticipantId> InParticipantId, const EOnlineTournamentParticipantType InParticipantType, const FVariantData& InScore)
+		: ParticipantId(InParticipantId)
+		, ParticipantType(InParticipantType)
+		, Score(InScore)
+	{
+	}
+
+public:
+	/** The ParticipantId who achieved Score */
+	TSharedRef<const FOnlineTournamentParticipantId> ParticipantId;
+	/** The type of participant this is */
+	EOnlineTournamentParticipantType ParticipantType;
+	/** The score for this participant */
+	FVariantData Score;
+};
+
+/**
+ * The results of a match
+ */
+struct FOnlineTournamentMatchResults
+{
+public:
+	struct FOnlineTournamentScreenshotData
+	{
+	public:
+		FOnlineTournamentScreenshotData(const FString& InScreenshotFormat, const TArray<uint8>& InScreenshotData)
+			: ScreenshotFormat(InScreenshotFormat)
+			, ScreenshotData(InScreenshotData)
+		{
+		}
+
+		FOnlineTournamentScreenshotData(FString&& InScreenshotFormat, TArray<uint8>&& InScreenshotData)
+			: ScreenshotFormat(MoveTemp(InScreenshotFormat))
+			, ScreenshotData(MoveTemp(InScreenshotData))
+		{
+		}
+
+	public:
+		/** The format of the screenshot stored in ScreenshotData */
+		FString ScreenshotFormat;
+		/** Raw bytes of a screenshot in the ScreenshotFormat format */
+		TArray<uint8> ScreenshotData;
+	};
+
+public:
+	/** Score data to submit */
+	TArray<FOnlineTournamentScore> ScoresToSubmit;
+
+	/** Optiona notes about a score */
+	TOptional<FString> Notes;
+	/** Optional screenshot data for proof of a score */
+	TOptional<FOnlineTournamentScreenshotData> Screenshot;
+};
+
+/**
+ * The details of a team in a tournament
+ */
+struct IOnlineTournamentTeamDetails
+	: public TSharedFromThis<IOnlineTournamentTeamDetails>
+{
+public:
+	virtual ~IOnlineTournamentTeamDetails() = default;
+
+	/** Get the TeamId of this team */
+	virtual TSharedRef<const FOnlineTournamentTeamId> GetTeamId() const = 0;
+	/** Get the player ids of this team (if they are known) */
+	virtual TOptional<TArray<TSharedRef<const FUniqueNetId>>> GetPlayerIds() const = 0;
+	/** Get the display name of this Team */
+	virtual FString GetDisplayName() const = 0;
+	/** Get an attribute for this team (varies by online platform) */
+	virtual TOptional<FVariantData> GetAttribute(const FName AttributeName) const = 0;
+};
+
+/**
+ * Filters to use when querying for participant information.
+ */
+struct FOnlineTournamentParticipantQueryFilter
+{
+	FOnlineTournamentParticipantQueryFilter(const EOnlineTournamentParticipantType InParticipantType)
+		: ParticipantType(InParticipantType)
+	{
+	}
+
+public:
+	EOnlineTournamentParticipantType ParticipantType;
+	TOptional<uint32> Limit;
+	TOptional<uint32> Offset;
+};
+
+enum class EOnlineTournamentParticipantState : uint8
+{
+	/** The participant has registered for the upcoming event */
+	Registered,
+	/** The participant has checked into the event and are ready to play */
+	CheckedIn,
+	/** The participant was present for the past event */
+	Present,
+	/** The participant was not present for the past event */
+	Absent
+};
+
+inline void LexFromString(TOptional<EOnlineTournamentParticipantState>& State, const TCHAR* const String)
+{
+	if (FCString::Stricmp(String, TEXT("Registered")) == 0)
+	{
+		State = EOnlineTournamentParticipantState::Registered;
+	}
+	else if (FCString::Stricmp(String, TEXT("CheckedIn")) == 0)
+	{
+		State = EOnlineTournamentParticipantState::CheckedIn;
+	}
+	else if (FCString::Stricmp(String, TEXT("Present")) == 0)
+	{
+		State = EOnlineTournamentParticipantState::Present;
+	}
+	else if (FCString::Stricmp(String, TEXT("Absent")) == 0)
+	{
+		State = EOnlineTournamentParticipantState::Absent;
+	}
+	else
+	{
+		State = TOptional<EOnlineTournamentParticipantState>();
+	}
+}
+
+inline FString LexToString(const EOnlineTournamentParticipantState ParticipantType)
+{
+	switch (ParticipantType)
+	{
+	case EOnlineTournamentParticipantState::Registered:
+		return TEXT("Registered");
+	case EOnlineTournamentParticipantState::CheckedIn:
+		return TEXT("CheckedIn");
+	case EOnlineTournamentParticipantState::Present:
+		return TEXT("Present");
+	case EOnlineTournamentParticipantState::Absent:
+		return TEXT("Absent");
+	}
+
+	checkNoEntry();
+	return FString();
+}
+
+/**
+ * The tournament-specific details of a participant in a tournament
+ */
+struct IOnlineTournamentParticipantDetails
+	: public TSharedFromThis<IOnlineTournamentParticipantDetails>
+{
+public:
+	virtual ~IOnlineTournamentParticipantDetails() = default;
+
+	/** Get the Tournament ID this participant is from */
+	virtual TSharedRef<const FOnlineTournamentId> GetTournamentId() const = 0;
+	/** Get the Player ID of this tournament participant (if applicable) */
+	virtual TSharedPtr<const FUniqueNetId> GetPlayerId() const = 0;
+	/** Get the Team ID of this tournament participant (if applicable) */
+	virtual TSharedPtr<const FOnlineTournamentTeamId> GetTeamId() const = 0;
+	/** Get the display name of this participant */
+	virtual const FString& GetDisplayName() const = 0;
+	/** Get the current state of the tournament participant */
+	virtual EOnlineTournamentParticipantState GetState() const = 0;
+	/** Get the current position of this tournament participant (if applicable) */
+	virtual TOptional<int32> GetPosition() const = 0;
+	/** Get the current score of this tournament participant (if applicable) */
+	virtual TOptional<FVariantData> GetScore() const = 0;
+	/** Get meta-data for this participant (varies based on online platform) */
+	virtual TOptional<FVariantData> GetAttribute(const FName AttributeName) const = 0;
+};
+
+/**
+	* States this match can be in
+	*/
+enum class EOnlineTournamentMatchState : uint8
+{
+	Created,
+	InProgress,
+	Finished
+};
+
+inline void LexFromString(TOptional<EOnlineTournamentMatchState>& State, const TCHAR* const String)
+{
+	if (FCString::Stricmp(String, TEXT("Created")) == 0)
+	{
+		State = EOnlineTournamentMatchState::Created;
+	}
+	else if (FCString::Stricmp(String, TEXT("InProgress")) == 0)
+	{
+		State = EOnlineTournamentMatchState::InProgress;
+	}
+	else if (FCString::Stricmp(String, TEXT("Finished")) == 0)
+	{
+		State = EOnlineTournamentMatchState::Finished;
+	}
+	else
+	{
+		State = TOptional<EOnlineTournamentMatchState>();
+	}
+}
+
+inline FString LexToString(const EOnlineTournamentMatchState ParticipantType)
+{
+	switch (ParticipantType)
+	{
+	case EOnlineTournamentMatchState::Created:
+		return TEXT("Created");
+	case EOnlineTournamentMatchState::InProgress:
+		return TEXT("InProgress");
+	case EOnlineTournamentMatchState::Finished:
+		return TEXT("Finished");
+	}
+
+	checkNoEntry();
+	return FString();
+}
+
+/**
+ * The details of a match
+ */
+struct IOnlineTournamentMatchDetails
+	: public TSharedFromThis<IOnlineTournamentMatchDetails>
+{
+public:
+
+	virtual ~IOnlineTournamentMatchDetails() = default;
+
+	/** Get the MatchId for this match*/
+	virtual TSharedRef<const FOnlineTournamentMatchId> GetMatchId() const = 0;
+	/** Get the type of participants for this match */
+	virtual EOnlineTournamentParticipantType GetParticipantType() const = 0;
+	/** Get the current state of this match */
+	virtual EOnlineTournamentMatchState GetMatchState() const = 0;
+	/** Get the bracket of this match */
+	virtual TOptional<FString> GetBracket() const = 0;
+	/** Get the round of this match */
+	virtual TOptional<int32> GetRound() const = 0;
+	/** Get the start time of this match in UTC */
+	virtual TOptional<FDateTime> GetStartDateUTC() const = 0;
+	/** Get the array of participants for this match */
+	virtual TArray<TSharedRef<const IOnlineTournamentParticipantDetails>> GetParticipants() const = 0;
+	/** Get meta-data for this match (varies based on online platform) */
+	virtual TOptional<FVariantData> GetAttribute(const FName AttributeName) const = 0;
+};
+
+/**
+ * The details of a tournament
+ */
+struct IOnlineTournamentDetails
+	: public TSharedFromThis<IOnlineTournamentDetails>
+{
+public:
+	virtual ~IOnlineTournamentDetails() = default;
+
+	/** Get the Tournament ID for this tournament */
+	virtual TSharedRef<const FOnlineTournamentId> GetTournamentId() const = 0;
+	/** Get the Title for this tournament */
+	virtual const FString& GetTitle() const = 0;
+	/** Get the Description for this tournament */
+	virtual const FString& GetDescription() const = 0;
+	/** Get the current state of this tournament */
+	virtual EOnlineTournamentState GetState() const = 0;
+	/** Get the the format of this tournament */
+	virtual EOnlineTournamentFormat GetFormat() const = 0;
+	/** Get the type of participants that are involved in this tournament */
+	virtual EOnlineTournamentParticipantType GetParticipantType() const = 0;
+	/** Get the list of Participant IDs for this tournament if known */
+	virtual const TArray<TSharedRef<const IOnlineTournamentParticipantDetails>>& GetParticipants() const = 0;
+	/** Get the registration start time of this tournament in UTC */
+	virtual TOptional<FDateTime> GetRegistrationStartDateUTC() const = 0;
+	/** Get the registration end time of this tournament in UTC */
+	virtual TOptional<FDateTime> GetRegistrationEndDateUTC() const = 0;
+	/** Get the start time of this tournament in UTC */
+	virtual TOptional<FDateTime> GetStartDateUTC() const = 0;
+	/** Get the check-in time of this tournament in UTC */
+	virtual TOptional<FTimespan> GetCheckInTimespan() const = 0;
+	/** Get the end time of this tournament in UTC */
+	virtual TOptional<FDateTime> GetEndDateUTC() const = 0;
+	/** Get the last time in UTC this tournament's details were updated */
+	virtual TOptional<FDateTime> GetLastUpdatedDateUTC() const = 0;
+	/** Does this tournament require a premium subscription to participate in? */
+	virtual TOptional<bool> RequiresPremiumSubscription() const = 0;
+	/** Get meta-data for this tournament (varies based on online platform) */
+	virtual TOptional<FVariantData> GetAttribute(const FName AttributeName) const = 0;
+};
+
+/**
+ * A delegate for when a tournament list has finished being queried.
+ *
+ * @param ResultStatus The result of the query
+ * @param TournamentIds An optional array of Tournament IDs
+ */
+DECLARE_DELEGATE_TwoParams(FOnlineTournamentQueryTournamentListComplete, const FOnlineError& /*ResultStatus*/, const TOptional<TArray<TSharedRef<const FOnlineTournamentId>>>& /*TournamentIds*/);
+
+/**
+ * A delegate for when tournament details have finished being queried.
+ *
+ * @param ResultStatus The result of the query
+ * @param TournamentDetails An optional array of Tournament Details
+ */
+DECLARE_DELEGATE_TwoParams(FOnlineTournamentQueryTournamentDetailsComplete, const FOnlineError& /*ResultStatus*/, const TOptional<TArray<TSharedRef<const IOnlineTournamentDetails>>>& /*TournamentDetails*/);
+
+/**
+ * A delegate for when a match list has finished being queried.
+ *
+ * @param ResultStatus The result of the query
+ * @param MatchIds An optional array of Match Ids
+ */
+DECLARE_DELEGATE_TwoParams(FOnlineTournamentQueryMatchListComplete, const FOnlineError& /*ResultStatus*/, const TOptional<TArray<TSharedRef<const FOnlineTournamentMatchId>>>& /*MatchIds*/);
+
+/**
+ * A delegate for when match details have finished being queried.
+ *
+ * @param ResultStatus The result of the query
+ * @param MatchDetails An optional array of Match Details
+ */
+DECLARE_DELEGATE_TwoParams(FOnlineTournamentQueryMatchDetailsComplete, const FOnlineError& /*ResultStatus*/, const TOptional<TArray<TSharedRef<const IOnlineTournamentMatchDetails>>>& /*MatchDetails*/);
+
+/**
+ * A delegate for when participant lists have finished being queried.
+ *
+ * @param ResultStatus The result of the query
+ * @param TotalResults Total possible results there are to be queried
+ * @param ParticipantList An optional array of Participant Details
+ */
+DECLARE_DELEGATE_ThreeParams(FOnlineTournamentQueryParticipantListComplete, const FOnlineError& /*ResultStatus*/, const TOptional<uint32> /*TotalResults*/, const TOptional<TArray<TSharedRef<const IOnlineTournamentParticipantDetails>>>& /*ParticipantList*/);
+
+/**
+ * A delegate for when team details have finished being queried.
+ *
+ * @param ResultStatus The result of the query
+ * @param TeamDetails An optional array of Team Details
+ */
+DECLARE_DELEGATE_TwoParams(FOnlineTournamentQueryTeamDetailsComplete, const FOnlineError& /*ResultStatus*/, const TOptional<TArray<TSharedRef<const IOnlineTournamentTeamDetails>>>& /*TeamDetails*/);
+
+/**
+ * A delegate for when match results have finished being submitted.
+ *
+ * @param ResultStatus The result of the query
+ */
+DECLARE_DELEGATE_OneParam(FOnlineTournamentSubmitMatchResultsComplete, const FOnlineError& /*ResultStatus*/);
+
+
+using FAdditionalMetaDataMap = TMap<FName, FString>;
+
+/**
+ * A delegate for when a tournament has been joined
+ *
+ * @param UserId The user who joined the tournament
+ * @param TournamentId The ID of the tournament joined
+ * @param AdditonalData A map of additional platform-specific data that was provided when this event was triggered
+ */
+DECLARE_MULTICAST_DELEGATE_ThreeParams(FOnOnlineTournamentTournamentJoined, const TSharedRef<const FUniqueNetId> /*UserId*/, const TSharedRef<const FOnlineTournamentId> /*TournamentId*/, const FAdditionalMetaDataMap& /*AdditionalData*/);
+using FOnOnlineTournamentTournamentJoinedDelegate = FOnOnlineTournamentTournamentJoined::FDelegate;
+
+/**
+ * A delegate for when a tournament has been joined
+ *
+ * @param UserId The user who joined the tournament
+ * @param MatchId The ID of the match joined
+ * @param AdditonalData A map of additional platform-specific data that was provided when this event was triggered
+ */
+DECLARE_MULTICAST_DELEGATE_ThreeParams(FOnOnlineTournamentMatchJoined, const TSharedRef<const FUniqueNetId> /*UserId*/, const TSharedRef<const FOnlineTournamentMatchId> /*MatchId*/, const FAdditionalMetaDataMap& /*AdditionalData*/);
+using FOnOnlineTournamentMatchJoinedDelegate = FOnOnlineTournamentMatchJoined::FDelegate;
+
+/**
+ * Interface to handle requesting and submitting information related to tournaments
+ */
+class IOnlineTournament
+	: public TSharedFromThis<IOnlineTournament, ESPMode::ThreadSafe>
+{
+public:
+	virtual ~IOnlineTournament() = default;
+
+	/**
+	 * Query a list of tournaments available for a user using specified filters.
+	 *
+	 * @param UserId The User to query tournaments as
+	 * @param QueryFilter The filter to user to include/exclude tournaments from the list.  One or more options MAY be required, depending on the online platform.
+	 * @param Delegate A delegate that is called when our tournament list query is complete
+	 */
+	virtual void QueryTournamentList(const TSharedRef<const FUniqueNetId> UserId, const FOnlineTournamentQueryFilter& QueryFilter, const FOnlineTournamentQueryTournamentListComplete& Delegate) = 0;
+
+	/**
+	 * Get a list of all Tournament IDs that have been queried by the specified user.
+	 *
+	 * @param UserId A User who previously queried tournaments
+	 * @return If previously queried by this user, an array of results from the query.
+	 */
+	virtual TArray<TSharedRef<const FOnlineTournamentId>> GetTournamentList(const TSharedRef<const FUniqueNetId> UserId) const = 0;
+
+	/**
+	 * Query tournament details from the perspective of the specified user.
+	 *
+	 * @param UserId A User to query tournament details as
+	 * @param TournamentIds An array of Tournament IDs to get more information about
+	 * @param Delegate A delegate that is called when our tournament details query is complete
+	 */
+	virtual void QueryTournamentDetails(const TSharedRef<const FUniqueNetId> UserId, const TArray<TSharedRef<const FOnlineTournamentId>>& TournamentIds, const FOnlineTournamentQueryTournamentDetailsComplete& Delegate) = 0;
+
+	/**
+	 * Get a tournament details result for a tournament that had been previously queried by the specified user.
+	 *
+	 * @param UserId A User who previously queried the tournament
+	 * @param TournamentId The Tournament ID that previously had details queried by UserId
+	 * @return If this tournament was previously queried by this user, the result from the query.
+	 */
+	virtual TSharedPtr<const IOnlineTournamentDetails> GetTournamentDetails(const TSharedRef<const FUniqueNetId> UserId, const TSharedRef<const FOnlineTournamentId> TournamentId) const = 0;
+
+	/**
+	 * Get tournament detail results for specified tournaments that have been previously queried by the specified user.
+	 *
+	 * @param UserId A User who previously queried tournaments
+	 * @param TournamentId The Tournament ID that previously had details queried by UserId
+	 * @return If tournaments have been previously queried by this user, any matching tournaments from the results are of those queries.
+	 */
+	virtual TArray<TSharedPtr<const IOnlineTournamentDetails>> GetTournamentDetails(const TSharedRef<const FUniqueNetId> UserId, const TArray<TSharedRef<const FOnlineTournamentId>>& TournamentIds) const = 0;
+
+	/**
+	 * Query a list of matches for a tournament from the perspective of the specified user.
+	 *
+	 * @param UserId A User to query match details as
+	 * @param TournamentId The Tournament ID to get Match IDs from
+	 * @param Delegate A delegate that is called when our match list query is complete
+	 */
+	virtual void QueryMatchList(const TSharedRef<const FUniqueNetId> UserId, const TSharedRef<const FOnlineTournamentId> TournamentId, const FOnlineTournamentQueryMatchListComplete& Delegate) = 0;
+
+	/**
+	 * Get match detail results for that have been previously queried by the specified user.
+	 *
+	 * @param UserId A User who previously queried match results
+	 * @param TournamentId The Tournament ID that previously had match results queried by UserId
+	 * @return If tournaments have been previously queried by this user, all results of those queries.
+	 */
+	virtual TArray<TSharedRef<const FOnlineTournamentMatchId>> GetMatchList(const TSharedRef<const FUniqueNetId> UserId, const TSharedRef<const FOnlineTournamentId> TournamentId) const = 0;
+
+	/**
+	 * Query match details for a tournament from the perspective of the specified user.
+	 *
+	 * @param UserId The User to query match details
+	 * @param MatchId The Match IDs to get more information about
+	 * @param Delegate A delegate that is called when our match details query is complete
+	 */
+	virtual void QueryMatchDetails(const TSharedRef<const FUniqueNetId> UserId, const TArray<TSharedRef<const FOnlineTournamentMatchId>>& MatchIds, const FOnlineTournamentQueryMatchDetailsComplete& Delegate) = 0;
+
+	/**
+	 * Get a match's details that have been previously queried by the specified user.
+	 *
+	 * @param UserId A User who previously queried the match
+	 * @param MatchId The Match ID that previously had details queried by UserId
+	 * @return If this match has been previously queried by this user, the result is of that query.
+	 */
+	virtual TSharedPtr<const IOnlineTournamentMatchDetails> GetMatchDetails(const TSharedRef<const FUniqueNetId> UserId, const TSharedRef<const FOnlineTournamentMatchId> MatchId) const = 0;
+
+	/**
+	 * Get match details that have been previously queried by the specified user.
+	 *
+	 * @param UserId A User who previously queried matches
+	 * @param MatchIds Match ids that previously had details queried by UserId
+	 * @return If matches have been previously queried by this user, any matching matches from the results are of those queries.
+	 */
+	virtual TArray<TSharedPtr<const IOnlineTournamentMatchDetails>> GetMatchDetails(const TSharedRef<const FUniqueNetId> UserId, const TArray<TSharedRef<const FOnlineTournamentMatchId>>& MatchIds) const = 0;
+
+	/**
+	 * Query a list of participants for a tournament from the perspective of the specified user.
+	 *
+	 * It is valid to request Team IDs or Player IDs from a Team tournament, and it is only valid to request Player IDs from an Individual tournament.
+	 *
+	 * @param UserId A User to query Participant IDs as
+	 * @param TournamentId The Tournament ID to get Participant IDs from
+	 * @param QueryFilter Filter to use to query participants.
+	 * @param Delegate A delegate that is called when our match list query is complete
+	 */
+	virtual void QueryParticipantList(const TSharedRef<const FUniqueNetId> UserId, const TSharedRef<const FOnlineTournamentId> TournamentId, const FOnlineTournamentParticipantQueryFilter& QueryFilter, const FOnlineTournamentQueryParticipantListComplete& Delegate) = 0;
+
+	/**
+	 * Get Participant details that have been previously queried by the specified user.
+	 *
+	 * @param UserId A User who previously queried participants
+	 * @param TournamentId The Tournament ID that previously had participants queried by UserId
+	 * @param ParticipantType The type of IDs to return from this tournament.  This type must have been previously queried.
+	 * @return If results have previously previously been queried by this user, the results of that query.
+	 */
+	virtual TArray<TSharedRef<const IOnlineTournamentParticipantDetails>> GetParticipantList(const TSharedRef<const FUniqueNetId> UserId, const TSharedRef<const FOnlineTournamentId> TournamentId, const EOnlineTournamentParticipantType ParticipantType) const = 0;
+
+	/**
+	 * Query team details from the perspective of the specified user.
+	 *
+	 * @param UserId A User to query team details as
+	 * @param TeamId The Team IDs to get more information about
+	 * @param Delegate A delegate that is called when our team details query is complete
+	 */
+	virtual void QueryTeamDetails(const TSharedRef<const FUniqueNetId> UserId, const TArray<TSharedRef<const FOnlineTournamentTeamId>>& TeamIds, const FOnlineTournamentQueryTeamDetailsComplete& Delegate) = 0;
+
+	/**
+	 * Get a teams's details that have been previously queried by the specified user.
+	 *
+	 * @param UserId A User who previously queried the team
+	 * @param TeamId The Team ID that previously had details queried by UserId
+	 * @return If this team has been previously queried by this user, the result is of that query.
+	 */
+	virtual TSharedPtr<const IOnlineTournamentTeamDetails> GetTeamDetails(const TSharedRef<const FUniqueNetId> UserId, const TSharedRef<const FOnlineTournamentTeamId> TeamId) const = 0;
+
+	/**
+	 * Get team details that have been previously queried by the specified user.
+	 *
+	 * @param UserId A User who previously queried teams
+	 * @param TeamIds The Team ID that previously had details queried by UserId
+	 * @return If teams have been previously queried by this user, any matching teams from the results are of those queries.
+	 */
+	virtual TArray<TSharedPtr<const IOnlineTournamentTeamDetails>> GetTeamDetails(const TSharedRef<const FUniqueNetId> UserId, const TArray<TSharedRef<const FOnlineTournamentTeamId>>& TeamIds) const = 0;
+
+	/**
+	 * Submit match results for a tournament match.
+	 *
+	 * @param UserId The User to report match results for
+	 * @param MatchId The Match Id to report match results for
+	 * @param MatchResults The results of the match
+	 * @param Delegate A delegate that is called when our match stats submissions is complete
+	 */
+	virtual void SubmitMatchResults(const TSharedRef<const FUniqueNetId> UserId, const TSharedRef<const FOnlineTournamentMatchId> MatchId, const FOnlineTournamentMatchResults& MatchResults, const FOnlineTournamentSubmitMatchResultsComplete& Delegate) = 0;
+
+	/**
+	 * Register for updates when a tournament has been joined
+	 *
+	 * @param Delegate Delegate to be called when a tournament has been joined
+	 */
+	virtual FDelegateHandle AddOnOnlineTournamentTournamentJoined(const FOnOnlineTournamentTournamentJoinedDelegate& Delegate) = 0;
+
+	/**
+	 * Unregister for tournament join updates using a previously-registered delegate handle
+	 *
+	 * @param DelegateHandle Delegate handle that was previously registered
+	 */
+	virtual void RemoveOnOnlineTournamentTournamentJoined(const FDelegateHandle& DelegateHandle) = 0;
+
+	/**
+	 * Register for updates when a tournament match has been joined
+	 *
+	 * @param Delegate Delegate to be called when a tournament match has been joined
+	 */
+	virtual FDelegateHandle AddOnOnlineTournamentMatchJoinedDelegate(const FOnOnlineTournamentMatchJoinedDelegate& Delegate) = 0;
+
+	/**
+	 * Unregister for tournament match join updates using a previously-registered delegate handle
+	 *
+	 * @param DelegateHandle Delegate handle that was previously registered
+	 */
+	virtual void RemoveOnOnlineTournamentMatchJoinedDelegate(const FDelegateHandle& DelegateHandle) = 0;
+
+#if !UE_BUILD_SHIPPING
+	/** Print all cached tournament information into the logs */
+	virtual void DumpCachedTournamentInfo(const TSharedRef<const FUniqueNetId> UserId) const = 0;
+	/** Print all cached match information into the logs */
+	virtual void DumpCachedMatchInfo(const TSharedRef<const FUniqueNetId> UserId) const = 0;
+	/** Print all cached participant information into the logs */
+	virtual void DumpCachedParticipantInfo(const TSharedRef<const FUniqueNetId> UserId) const = 0;
+	/** Print all cached team information into the logs */
+	virtual void DumpCachedTeamInfo(const TSharedRef<const FUniqueNetId> UserId) const = 0;
+#endif // !UE_BUILD_SHIPPING
+};
diff --git a/Engine/Plugins/Online/OnlineSubsystem/Source/Public/OnlineSubsystem.h b/Engine/Plugins/Online/OnlineSubsystem/Source/Public/OnlineSubsystem.h
--- a/Engine/Plugins/Online/OnlineSubsystem/Source/Public/OnlineSubsystem.h
+++ b/Engine/Plugins/Online/OnlineSubsystem/Source/Public/OnlineSubsystem.h
@@ -90,6 +90,7 @@ typedef TSharedPtr<class IOnlineMessage, ESPMode::ThreadSafe> IOnlineMessagePtr;
 typedef TSharedPtr<class IOnlinePresence, ESPMode::ThreadSafe> IOnlinePresencePtr;
 typedef TSharedPtr<class IOnlineChat, ESPMode::ThreadSafe> IOnlineChatPtr;
 typedef TSharedPtr<class IOnlineTurnBased, ESPMode::ThreadSafe> IOnlineTurnBasedPtr;
+typedef TSharedPtr<class IOnlineTournament, ESPMode::ThreadSafe> IOnlineTournamentPtr;
 typedef TSharedPtr<class FOnlineNotificationHandler, ESPMode::ThreadSafe> FOnlineNotificationHandlerPtr;
 typedef TSharedPtr<class FOnlineNotificationTransportManager, ESPMode::ThreadSafe> FOnlineNotificationTransportManagerPtr;
 
@@ -385,19 +386,19 @@ public:
 
 	/** 
 	 * Get the interface for accessing online messages
-	 * @return Interface pointer for the appropriate online user service
+	 * @return Interface pointer for the appropriate online message service
 	 */
 	virtual IOnlineMessagePtr GetMessageInterface() const = 0;
 
 	/** 
 	 * Get the interface for managing rich presence information
-	 * @return Interface pointer for the appropriate online user service
+	 * @return Interface pointer for the appropriate online presence service
 	 */
 	virtual IOnlinePresencePtr GetPresenceInterface() const = 0;
 
 	/** 
 	 * Get the interface for user-user and user-room chat functionality
-	 * @return Interface pointer for the appropriate online user service
+	 * @return Interface pointer for the appropriate online chat service
 	 */
 	virtual IOnlineChatPtr GetChatInterface() const = 0;
 
@@ -412,10 +413,16 @@ public:
 
 	/**
 	 * Get the interface for managing turn based multiplayer games
-	 * @return Interface pointer for the appropriate online user service
+	 * @return Interface pointer for the appropriate online turn-based service
 	 */
 	virtual IOnlineTurnBasedPtr GetTurnBasedInterface() const = 0;
 
+	/**
+	 * Get the interface for managing tournament information
+	 * @return Interface pointer for the appropriate online tournament service
+	 */
+	virtual IOnlineTournamentPtr GetTournamentInterface() const = 0;
+
 	/**
 	 * Get the transport manager instance for this subsystem
 	 * @return Pointer for the appropriate transport manager
diff --git a/Engine/Plugins/Online/OnlineSubsystemAmazon/Source/Public/OnlineSubsystemAmazon.h b/Engine/Plugins/Online/OnlineSubsystemAmazon/Source/Public/OnlineSubsystemAmazon.h
--- a/Engine/Plugins/Online/OnlineSubsystemAmazon/Source/Public/OnlineSubsystemAmazon.h
+++ b/Engine/Plugins/Online/OnlineSubsystemAmazon/Source/Public/OnlineSubsystemAmazon.h
@@ -58,7 +58,8 @@ public:
 	virtual IOnlinePresencePtr GetPresenceInterface() const override { return nullptr; }
 	virtual IOnlineChatPtr GetChatInterface() const override { return nullptr; }
 	virtual IOnlineTurnBasedPtr GetTurnBasedInterface() const override { return nullptr; }
-	
+	virtual IOnlineTournamentPtr GetTournamentInterface() const override { return nullptr; }
+
 	virtual bool Init() override;
 	virtual bool Shutdown() override;
 	virtual FString GetAppId() const override;
diff --git a/Engine/Plugins/Online/OnlineSubsystemFacebook/Source/Private/OnlineSubsystemFacebookCommon.cpp b/Engine/Plugins/Online/OnlineSubsystemFacebook/Source/Private/OnlineSubsystemFacebookCommon.cpp
--- a/Engine/Plugins/Online/OnlineSubsystemFacebook/Source/Private/OnlineSubsystemFacebookCommon.cpp
+++ b/Engine/Plugins/Online/OnlineSubsystemFacebook/Source/Private/OnlineSubsystemFacebookCommon.cpp
@@ -210,6 +210,11 @@ IOnlineTurnBasedPtr FOnlineSubsystemFacebookCommon::GetTurnBasedInterface() cons
 	return nullptr;
 }
 
+IOnlineTournamentPtr FOnlineSubsystemFacebookCommon::GetTournamentInterface() const
+{
+	return nullptr;
+}
+
 FText FOnlineSubsystemFacebookCommon::GetOnlineServiceName() const
 {
 	return NSLOCTEXT("OnlineSubsystemFacebook", "OnlineServiceName", "Facebook");
diff --git a/Engine/Plugins/Online/OnlineSubsystemFacebook/Source/Private/OnlineSubsystemFacebookCommon.h b/Engine/Plugins/Online/OnlineSubsystemFacebook/Source/Private/OnlineSubsystemFacebookCommon.h
--- a/Engine/Plugins/Online/OnlineSubsystemFacebook/Source/Private/OnlineSubsystemFacebookCommon.h
+++ b/Engine/Plugins/Online/OnlineSubsystemFacebook/Source/Private/OnlineSubsystemFacebookCommon.h
@@ -47,6 +47,7 @@ public:
 	virtual IOnlinePresencePtr GetPresenceInterface() const override;
 	virtual IOnlineChatPtr GetChatInterface() const override;
 	virtual IOnlineTurnBasedPtr GetTurnBasedInterface() const override;
+	virtual IOnlineTournamentPtr GetTournamentInterface() const override;
 	virtual bool Init() override;
 	virtual bool Shutdown() override;
 	virtual FString GetAppId() const override;
diff --git a/Engine/Plugins/Online/OnlineSubsystemGoogle/Source/Private/OnlineSubsystemGoogleCommon.cpp b/Engine/Plugins/Online/OnlineSubsystemGoogle/Source/Private/OnlineSubsystemGoogleCommon.cpp
--- a/Engine/Plugins/Online/OnlineSubsystemGoogle/Source/Private/OnlineSubsystemGoogleCommon.cpp
+++ b/Engine/Plugins/Online/OnlineSubsystemGoogle/Source/Private/OnlineSubsystemGoogleCommon.cpp
@@ -231,6 +231,11 @@ IOnlineTurnBasedPtr FOnlineSubsystemGoogleCommon::GetTurnBasedInterface() const
 	return nullptr;
 }
 
+IOnlineTournamentPtr FOnlineSubsystemGoogleCommon::GetTournamentInterface() const
+{
+	return nullptr;
+}
+
 FText FOnlineSubsystemGoogleCommon::GetOnlineServiceName() const
 {
 	return NSLOCTEXT("OnlineSubsystemGoogleCommon", "OnlineServiceName", "Google");
diff --git a/Engine/Plugins/Online/OnlineSubsystemGoogle/Source/Public/OnlineSubsystemGoogleCommon.h b/Engine/Plugins/Online/OnlineSubsystemGoogle/Source/Public/OnlineSubsystemGoogleCommon.h
--- a/Engine/Plugins/Online/OnlineSubsystemGoogle/Source/Public/OnlineSubsystemGoogleCommon.h
+++ b/Engine/Plugins/Online/OnlineSubsystemGoogle/Source/Public/OnlineSubsystemGoogleCommon.h
@@ -60,6 +60,7 @@ public:
 	virtual IOnlinePresencePtr GetPresenceInterface() const override;
 	virtual IOnlineChatPtr GetChatInterface() const override;
 	virtual IOnlineTurnBasedPtr GetTurnBasedInterface() const override;
+	virtual IOnlineTournamentPtr GetTournamentInterface() const override;
 	virtual FText GetOnlineServiceName() const override;
 	virtual bool Init() override;
 	virtual bool Shutdown() override;
diff --git a/Engine/Plugins/Online/OnlineSubsystemNull/Source/Private/OnlineSubsystemNull.cpp b/Engine/Plugins/Online/OnlineSubsystemNull/Source/Private/OnlineSubsystemNull.cpp
--- a/Engine/Plugins/Online/OnlineSubsystemNull/Source/Private/OnlineSubsystemNull.cpp
+++ b/Engine/Plugins/Online/OnlineSubsystemNull/Source/Private/OnlineSubsystemNull.cpp
@@ -129,7 +129,12 @@ IOnlineChatPtr FOnlineSubsystemNull::GetChatInterface() const
 
 IOnlineTurnBasedPtr FOnlineSubsystemNull::GetTurnBasedInterface() const
 {
-    return nullptr;
+	return nullptr;
+}
+
+IOnlineTournamentPtr FOnlineSubsystemNull::GetTournamentInterface() const
+{
+	return nullptr;
 }
 
 bool FOnlineSubsystemNull::Tick(float DeltaTime)
diff --git a/Engine/Plugins/Online/OnlineSubsystemNull/Source/Public/OnlineSubsystemNull.h b/Engine/Plugins/Online/OnlineSubsystemNull/Source/Public/OnlineSubsystemNull.h
--- a/Engine/Plugins/Online/OnlineSubsystemNull/Source/Public/OnlineSubsystemNull.h
+++ b/Engine/Plugins/Online/OnlineSubsystemNull/Source/Public/OnlineSubsystemNull.h
@@ -63,7 +63,8 @@ public:
 	virtual IOnlinePresencePtr GetPresenceInterface() const override;
 	virtual IOnlineChatPtr GetChatInterface() const override;
 	virtual IOnlineTurnBasedPtr GetTurnBasedInterface() const override;
-	
+	virtual IOnlineTournamentPtr GetTournamentInterface() const override;
+
 	virtual bool Init() override;
 	virtual bool Shutdown() override;
 	virtual FString GetAppId() const override;
diff --git a/Engine/Plugins/Online/OnlineSubsystemOculus/Source/Private/OnlineSubsystemOculus.cpp b/Engine/Plugins/Online/OnlineSubsystemOculus/Source/Private/OnlineSubsystemOculus.cpp
--- a/Engine/Plugins/Online/OnlineSubsystemOculus/Source/Private/OnlineSubsystemOculus.cpp
+++ b/Engine/Plugins/Online/OnlineSubsystemOculus/Source/Private/OnlineSubsystemOculus.cpp
@@ -139,6 +139,11 @@ IOnlineTurnBasedPtr FOnlineSubsystemOculus::GetTurnBasedInterface() const
 	return nullptr;
 }
 
+IOnlineTournamentPtr FOnlineSubsystemOculus::GetTournamentInterface() const
+{
+	return nullptr;
+}
+
 bool FOnlineSubsystemOculus::Tick(float DeltaTime)
 {
 	if (!FOnlineSubsystemImpl::Tick(DeltaTime))
diff --git a/Engine/Plugins/Online/OnlineSubsystemOculus/Source/Public/OnlineSubsystemOculus.h b/Engine/Plugins/Online/OnlineSubsystemOculus/Source/Public/OnlineSubsystemOculus.h
--- a/Engine/Plugins/Online/OnlineSubsystemOculus/Source/Public/OnlineSubsystemOculus.h
+++ b/Engine/Plugins/Online/OnlineSubsystemOculus/Source/Public/OnlineSubsystemOculus.h
@@ -59,6 +59,7 @@ public:
 	virtual IOnlinePresencePtr GetPresenceInterface() const override;
 	virtual IOnlineChatPtr GetChatInterface() const override;
 	virtual IOnlineTurnBasedPtr GetTurnBasedInterface() const override;
+	virtual IOnlineTournamentPtr GetTournamentInterface() const override;
 
 	virtual bool Init() override;
 	virtual bool Shutdown() override;
diff --git a/Engine/Plugins/Online/OnlineSubsystemSteam/Source/Private/OnlineSubsystemSteam.cpp b/Engine/Plugins/Online/OnlineSubsystemSteam/Source/Private/OnlineSubsystemSteam.cpp
--- a/Engine/Plugins/Online/OnlineSubsystemSteam/Source/Private/OnlineSubsystemSteam.cpp
+++ b/Engine/Plugins/Online/OnlineSubsystemSteam/Source/Private/OnlineSubsystemSteam.cpp
@@ -304,7 +304,12 @@ IOnlineChatPtr FOnlineSubsystemSteam::GetChatInterface() const
 
 IOnlineTurnBasedPtr FOnlineSubsystemSteam::GetTurnBasedInterface() const
 {
-    return nullptr;
+	return nullptr;
+}
+
+IOnlineTournamentPtr FOnlineSubsystemSteam::GetTournamentInterface() const
+{
+	return nullptr;
 }
 
 void FOnlineSubsystemSteam::QueueAsyncTask(FOnlineAsyncTask* AsyncTask)
diff --git a/Engine/Plugins/Online/OnlineSubsystemSteam/Source/Public/OnlineSubsystemSteam.h b/Engine/Plugins/Online/OnlineSubsystemSteam/Source/Public/OnlineSubsystemSteam.h
--- a/Engine/Plugins/Online/OnlineSubsystemSteam/Source/Public/OnlineSubsystemSteam.h
+++ b/Engine/Plugins/Online/OnlineSubsystemSteam/Source/Public/OnlineSubsystemSteam.h
@@ -248,6 +248,7 @@ public:
 	virtual IOnlinePresencePtr GetPresenceInterface() const override;
 	virtual IOnlineChatPtr GetChatInterface() const override;
 	virtual IOnlineTurnBasedPtr GetTurnBasedInterface() const override;
+	virtual IOnlineTournamentPtr GetTournamentInterface() const override;
 	virtual bool IsLocalPlayer(const FUniqueNetId& UniqueId) const override;
 	virtual bool Init() override;
 	virtual bool Shutdown() override;
diff --git a/Engine/Plugins/Online/OnlineSubsystemTwitch/Source/Private/OnlineSubsystemTwitch.cpp b/Engine/Plugins/Online/OnlineSubsystemTwitch/Source/Private/OnlineSubsystemTwitch.cpp
--- a/Engine/Plugins/Online/OnlineSubsystemTwitch/Source/Private/OnlineSubsystemTwitch.cpp
+++ b/Engine/Plugins/Online/OnlineSubsystemTwitch/Source/Private/OnlineSubsystemTwitch.cpp
@@ -126,6 +126,11 @@ IOnlineTurnBasedPtr FOnlineSubsystemTwitch::GetTurnBasedInterface() const
 	return nullptr;
 }
 
+IOnlineTournamentPtr FOnlineSubsystemTwitch::GetTournamentInterface() const
+{
+	return nullptr;
+}
+
 bool FOnlineSubsystemTwitch::Init()
 {
 	UE_LOG_ONLINE(VeryVerbose, TEXT("FOnlineSubsystemTwitch::Init() Name: %s"), *InstanceName.ToString());
diff --git a/Engine/Plugins/Online/OnlineSubsystemTwitch/Source/Public/OnlineSubsystemTwitch.h b/Engine/Plugins/Online/OnlineSubsystemTwitch/Source/Public/OnlineSubsystemTwitch.h
--- a/Engine/Plugins/Online/OnlineSubsystemTwitch/Source/Public/OnlineSubsystemTwitch.h
+++ b/Engine/Plugins/Online/OnlineSubsystemTwitch/Source/Public/OnlineSubsystemTwitch.h
@@ -33,7 +33,7 @@ public:
 	virtual IOnlineEntitlementsPtr GetEntitlementsInterface() const override;
 	virtual IOnlineLeaderboardsPtr GetLeaderboardsInterface() const override;
 	virtual IOnlineVoicePtr GetVoiceInterface() const override;
-	virtual IOnlineExternalUIPtr GetExternalUIInterface() const override;	
+	virtual IOnlineExternalUIPtr GetExternalUIInterface() const override;
 	virtual IOnlineTimePtr GetTimeInterface() const override;
 	virtual IOnlineIdentityPtr GetIdentityInterface() const override;
 	virtual IOnlineTitleFilePtr GetTitleFileInterface() const override;
@@ -48,6 +48,7 @@ public:
 	virtual IOnlinePresencePtr GetPresenceInterface() const override;
 	virtual IOnlineChatPtr GetChatInterface() const override;
 	virtual IOnlineTurnBasedPtr GetTurnBasedInterface() const override;
+	virtual IOnlineTournamentPtr GetTournamentInterface() const override;
 	virtual bool Init() override;
 	virtual void PreUnload() override;
 	virtual bool Shutdown() override;
