Make MakeJoinInfoFromToken return a TSharedPtr as the input may be invalid

diff --git a/Engine/Plugins/Online/OnlineSubsystem/Source/Public/Interfaces/OnlinePartyInterface.h b/Engine/Plugins/Online/OnlineSubsystem/Source/Public/Interfaces/OnlinePartyInterface.h
--- a/Engine/Plugins/Online/OnlineSubsystem/Source/Public/Interfaces/OnlinePartyInterface.h
+++ b/Engine/Plugins/Online/OnlineSubsystem/Source/Public/Interfaces/OnlinePartyInterface.h
@@ -114,8 +114,24 @@ public:
 	}
 
 	/**
-	* Returns true if there are any dirty keys
-	*/
+	 * Set an attribute from the party data
+	 *
+	 * @param AttrName - key for the attribute
+	 * @param AttrValue - value to set the attribute to
+	 */
+	void SetAttribute(FString&& AttrName, FVariantData&& AttrValue)
+	{
+		FVariantData& NewAttrValue = KeyValAttrs.FindOrAdd(AttrName);
+		if (NewAttrValue != AttrValue)
+		{
+			NewAttrValue = MoveTemp(AttrValue);
+			DirtyKeys.Emplace(MoveTemp(AttrName));
+		}
+	}
+
+	/**
+	 * Returns true if there are any dirty keys
+	 */
 	bool HasDirtyKeys() const
 	{
 		return DirtyKeys.Num() > 0;
@@ -227,6 +243,7 @@ public:
  * Info needed to join a party
  */
 class IOnlinePartyJoinInfo
+	: public TSharedFromThis<IOnlinePartyJoinInfo>
 {
 public:
 	IOnlinePartyJoinInfo() {}
@@ -409,27 +426,26 @@ enum class EPartyState
  */
 class FOnlineParty : public TSharedFromThis<FOnlineParty>
 {
+	FOnlineParty() = delete;
 protected:
-	FOnlineParty();
-	explicit FOnlineParty(const TSharedRef<const FOnlinePartyId>& InPartyId, const FOnlinePartyTypeId InPartyTypeId)
+	FOnlineParty(const TSharedRef<const FOnlinePartyId>& InPartyId, const FOnlinePartyTypeId InPartyTypeId, TSharedRef<FPartyConfiguration> InConfig = MakeShared<FPartyConfiguration>())
 		: PartyId(InPartyId)
 		, PartyTypeId(InPartyTypeId)
 		, State(EPartyState::None)
-		, Config(MakeShareable(new FPartyConfiguration()))
+		, Config(InConfig)
 	{}
 
 public:
-	virtual ~FOnlineParty()
-	{}
+	virtual ~FOnlineParty() = default;
 
 	virtual bool CanLocalUserInvite(const FUniqueNetId& LocalUserId) const = 0;
 	virtual bool IsJoinable() const = 0;
 
-	/** unique id of the party */
+	/** Unique id of the party */
 	TSharedRef<const FOnlinePartyId> PartyId;
-	/** unique id of the party */
+	/** Type of party (e.g., Primary) */
 	const FOnlinePartyTypeId PartyTypeId;
-	/** unique id of the leader */
+	/** Unique id of the leader */
 	TSharedPtr<const FUniqueNetId> LeaderId;
 	/** The current state of the party */
 	EPartyState State;
@@ -1203,7 +1219,7 @@ public:
 	 *
 	 * return the new IOnlinePartyJoinInfo object
 	 */
-	virtual TSharedRef<IOnlinePartyJoinInfo> MakeJoinInfoFromToken(const FString& Token) const = 0;
+	virtual TSharedPtr<IOnlinePartyJoinInfo> MakeJoinInfoFromToken(const FString& Token) const = 0;
 
 	/**
 	 * Checks to see if there is a pending command line invite and consumes it
