Add some helper functions for just getting the FOnlinePartyData attributes in JSON format

diff --git a/Engine/Plugins/Online/OnlineSubsystem/Source/Private/OnlinePartyInterface.cpp b/Engine/Plugins/Online/OnlineSubsystem/Source/Private/OnlinePartyInterface.cpp
--- a/Engine/Plugins/Online/OnlineSubsystem/Source/Private/OnlinePartyInterface.cpp
+++ b/Engine/Plugins/Online/OnlineSubsystem/Source/Private/OnlinePartyInterface.cpp
@@ -28,14 +28,7 @@ void FOnlinePartyData::ToJsonFull(FString& JsonString) const
 
 	// iterate over key/val attrs and convert each entry to a json string
 	TSharedRef<FJsonObject> JsonObject(new FJsonObject());
-	TSharedRef<FJsonObject> JsonProperties = MakeShared<FJsonObject>();
-	for (auto Iterator : KeyValAttrs)
-	{
-		const FString& PropertyName = Iterator.Key;
-		const FVariantData& PropertyValue = Iterator.Value;
-
-		PropertyValue.AddToJsonObject(JsonProperties, PropertyName);
-	}
+	TSharedRef<FJsonObject> JsonProperties = GetAllAttributesAsJsonObject();
 	JsonObject->SetNumberField(TEXT("Rev"), RevisionCount);
 	JsonObject->SetObjectField(TEXT("Attrs"), JsonProperties);
 
@@ -66,6 +59,30 @@ void FOnlinePartyData::ToJsonDirty(FString& JsonString) const
 	JsonWriter->Close();
 }
 
+TSharedRef<FJsonObject> FOnlinePartyData::GetAllAttributesAsJsonObject() const
+{
+	TSharedRef<FJsonObject> JsonProperties = MakeShared<FJsonObject>();
+	for (auto Iterator : KeyValAttrs)
+	{
+		const FString& PropertyName = Iterator.Key;
+		const FVariantData& PropertyValue = Iterator.Value;
+
+		PropertyValue.AddToJsonObject(JsonProperties, PropertyName);
+	}
+	return JsonProperties;
+}
+
+FString FOnlinePartyData::GetAllAttributesAsJsonObjectString() const
+{
+	TSharedRef<FJsonObject> JsonProperties = GetAllAttributesAsJsonObject();
+
+	FString JsonString;
+	auto JsonWriter = TJsonWriterFactory<TCHAR, TCondensedJsonPrintPolicy<TCHAR> >::Create(&JsonString);
+	FJsonSerializer::Serialize(JsonProperties, JsonWriter);
+	JsonWriter->Close();
+	return JsonString;
+}
+
 void FOnlinePartyData::FromJson(const FString& JsonString)
 {
 	// json string to key/val attrs
diff --git a/Engine/Plugins/Online/OnlineSubsystem/Source/Public/Interfaces/OnlinePartyInterface.h b/Engine/Plugins/Online/OnlineSubsystem/Source/Public/Interfaces/OnlinePartyInterface.h
--- a/Engine/Plugins/Online/OnlineSubsystem/Source/Public/Interfaces/OnlinePartyInterface.h
+++ b/Engine/Plugins/Online/OnlineSubsystem/Source/Public/Interfaces/OnlinePartyInterface.h
@@ -176,7 +176,6 @@ public:
 	 * Generate a JSON packet containing all key-value attributes
 	 * 
 	 * @param JsonString - [out] string containing the resulting JSON output
-	 * 
 	 */
 	void ToJsonFull(FString& JsonString) const;
 	
@@ -184,15 +183,25 @@ public:
 	 * Generate a JSON packet containing only the dirty key-value attributes for a delta update
 	 *
 	 * @param JsonString - [out] string containing the resulting JSON output
-	 *
 	 */
 	void ToJsonDirty(FString& JsonString) const;
 
+	/**
+	 * Create a JSON object containing all key-value attributes
+	 * @return a JSON object containing all key-value attributes
+	 */
+	TSharedRef<FJsonObject> GetAllAttributesAsJsonObject() const;
+
+	/**
+	 * Create a string representing a JSON object containing all key-value attributes
+	 * @return a string representing a JSON object containing all key-value attributes
+	 */
+	FString GetAllAttributesAsJsonObjectString() const;
+
 	/** 
 	 * Update attributes from a JSON packet
 	 *
 	 * @param JsonString - string containing the JSON packet
-	 *
 	 */
 	void FromJson(const FString& JsonString);
 
