Added back in checks for not using mcp instancing in Server and Client Modes.

diff --git a/Engine/Plugins/Online/OnlineSubsystem/Source/Private/OnlineSubsystemModule.cpp b/Engine/Plugins/Online/OnlineSubsystem/Source/Private/OnlineSubsystemModule.cpp
--- a/Engine/Plugins/Online/OnlineSubsystem/Source/Private/OnlineSubsystemModule.cpp
+++ b/Engine/Plugins/Online/OnlineSubsystem/Source/Private/OnlineSubsystemModule.cpp
@@ -233,6 +233,7 @@ void FOnlineSubsystemModule::EnumerateOnlineSubsystems(FEnumerateOnlineSubsystem
 
 FName FOnlineSubsystemModule::ParseOnlineSubsystemName(const FName& FullName, FName& SubsystemName, FName& InstanceName) const
 {
+#if !(UE_GAME || UE_SERVER)
 	SubsystemName = DefaultPlatformService;
 	InstanceName = FOnlineSubsystemImpl::DefaultInstanceName;
 
@@ -261,6 +262,18 @@ FName FOnlineSubsystemModule::ParseOnlineSubsystemName(const FName& FullName, FN
 	}
 
 	return FName(*FString::Printf(TEXT("%s:%s"), *SubsystemName.ToString(), *InstanceName.ToString()));
+#else	
+	
+	SubsystemName = FullName.IsNone() ? DefaultPlatformService : FullName;
+	InstanceName = FOnlineSubsystemImpl::DefaultInstanceName;
+
+#if !UE_BUILD_SHIPPING
+	int32 DelimIdx = INDEX_NONE;
+	static const TCHAR InstanceDelim = ':';
+	ensure(!FullName.ToString().FindChar(InstanceDelim, DelimIdx) && DelimIdx == INDEX_NONE);
+#endif
+	return SubsystemName;
+#endif // !(UE_GAME || UE_SERVER)
 }
 
 IOnlineSubsystem* FOnlineSubsystemModule::GetOnlineSubsystem(const FName InSubsystemName)
