Add accessor functions to get more insight into dirty keys

diff --git a/Engine/Plugins/Online/OnlineSubsystem/Source/Private/OnlinePartyInterface.cpp b/Engine/Plugins/Online/OnlineSubsystem/Source/Private/OnlinePartyInterface.cpp
--- a/Engine/Plugins/Online/OnlineSubsystem/Source/Private/OnlinePartyInterface.cpp
+++ b/Engine/Plugins/Online/OnlineSubsystem/Source/Private/OnlinePartyInterface.cpp
@@ -22,6 +22,19 @@ bool FOnlinePartyData::operator!=(const FOnlinePartyData& Other) const
 	return !operator==(Other);
 }
 
+FOnlineKeyValuePairs<FString, FVariantData> FOnlinePartyData::GetDirtyKeyVaAttrs() const
+{
+	FOnlineKeyValuePairs<FString, FVariantData> DirtyKeyValAttrs;
+	for (const FString& PropertyName : DirtyKeys)
+	{
+		const FVariantData* PropertyValue = KeyValAttrs.Find(PropertyName);
+		check(PropertyValue);
+
+		DirtyKeyValAttrs.Emplace(PropertyName, *PropertyValue);
+	}
+	return DirtyKeyValAttrs;
+}
+
 void FOnlinePartyData::ToJsonFull(FString& JsonString) const
 {
 	JsonString.Empty();
@@ -44,7 +57,7 @@ void FOnlinePartyData::ToJsonDirty(FString& JsonString) const
 	// iterate over key/val attrs and convert each entry to a json string
 	TSharedRef<FJsonObject> JsonObject(new FJsonObject());
 	TSharedRef<FJsonObject> JsonProperties = MakeShared<FJsonObject>();
-	for (auto& PropertyName : DirtyKeys)
+	for (const FString& PropertyName : DirtyKeys)
 	{
 		const FVariantData* PropertyValue = KeyValAttrs.Find(PropertyName);
 		check(PropertyValue);
@@ -62,7 +75,7 @@ void FOnlinePartyData::ToJsonDirty(FString& JsonString) const
 TSharedRef<FJsonObject> FOnlinePartyData::GetAllAttributesAsJsonObject() const
 {
 	TSharedRef<FJsonObject> JsonProperties = MakeShared<FJsonObject>();
-	for (auto Iterator : KeyValAttrs)
+	for (const TPair<FString, FVariantData>& Iterator : KeyValAttrs)
 	{
 		const FString& PropertyName = Iterator.Key;
 		const FVariantData& PropertyValue = Iterator.Value;
@@ -94,7 +107,7 @@ void FOnlinePartyData::FromJson(const FString& JsonString)
 		if (JsonObject->HasTypedField<EJson::Object>(TEXT("Attrs")))
 		{
 			const TSharedPtr<FJsonObject>& JsonProperties = JsonObject->GetObjectField(TEXT("Attrs"));
-			for (auto& JsonProperty : JsonProperties->Values)
+			for (const TPair<FString, TSharedPtr<FJsonValue>>& JsonProperty : JsonProperties->Values)
 			{
 				FString PropertyName;
 				FVariantData PropertyData;
diff --git a/Engine/Plugins/Online/OnlineSubsystem/Source/Public/Interfaces/OnlinePartyInterface.h b/Engine/Plugins/Online/OnlineSubsystem/Source/Public/Interfaces/OnlinePartyInterface.h
--- a/Engine/Plugins/Online/OnlineSubsystem/Source/Public/Interfaces/OnlinePartyInterface.h
+++ b/Engine/Plugins/Online/OnlineSubsystem/Source/Public/Interfaces/OnlinePartyInterface.h
@@ -130,13 +130,32 @@ public:
 	}
 
 	/**
-	 * Returns true if there are any dirty keys
+	 * Mark an attribute as dirty so it can be rebroadcasted
+	 *
+	 * @param AttrName - key for the attribute to mark dirty
+	 */
+	void MarkAttributeDirty(FString&& AttrName)
+	{
+		DirtyKeys.Emplace(MoveTemp(AttrName));
+	}
+
+	/**
+	 * Check if there are any dirty keys
+	 *
+	 * @return true if there are any dirty keys
 	 */
 	bool HasDirtyKeys() const
 	{
 		return DirtyKeys.Num() > 0;
 	}
 
+	/**
+	 * Get the dirty key-value attributes
+	 *
+	 * @return the dirty key-value attributes
+	 */
+	FOnlineKeyValuePairs<FString, FVariantData> GetDirtyKeyVaAttrs() const;
+
 	/**
 	 * Clear the attributes map
 	 */
@@ -222,7 +241,7 @@ public:
 
 private:
 	/** map of key/val attributes that represents the data */
-	FOnlineKeyValuePairs<FString, FVariantData>  KeyValAttrs;
+	FOnlineKeyValuePairs<FString, FVariantData> KeyValAttrs;
 
 	/** set of which fields are dirty and need to transmitted */
 	TSet<FString> DirtyKeys;
