Deprecate IOnlinePartySystem approved rejoins.  This is an internal implementation detail and does not make sense in the context of the party system common interface.

diff --git a/Engine/Plugins/Online/OnlineFramework/Source/Party/Private/Party/SocialParty.cpp b/Engine/Plugins/Online/OnlineFramework/Source/Party/Private/Party/SocialParty.cpp
--- a/Engine/Plugins/Online/OnlineFramework/Source/Party/Private/Party/SocialParty.cpp
+++ b/Engine/Plugins/Online/OnlineFramework/Source/Party/Private/Party/SocialParty.cpp
@@ -522,7 +522,6 @@ UPartyMember* USocialParty::GetOrCreatePartyMember(const FUniqueNetId& MemberId)
 					PartyMembersById.Add(MemberIdRepl, PartyMember);
 					PartyMember->InitializePartyMember(OssPartyMember.ToSharedRef(), FSimpleDelegate::CreateUObject(this, &USocialParty::HandleMemberInitialized, PartyMember));
 
-					PartyInterface->ApproveUserForRejoin(*OwningLocalUserId, PartyId, MemberId);
 					RefreshPublicJoinability();
 
 					OnPartyMemberCreated().Broadcast(*PartyMember);
@@ -890,15 +889,6 @@ void USocialParty::HandlePartyMemberExited(const FUniqueNetId& LocalUserId, cons
 				// Update party join state, will cause a failure on leader promotion currently
 				// because we can't tell the difference between "expected leader" and "actually the new leader"
 				RefreshPublicJoinability();
-
-				if (ExitReason != EMemberExitedReason::Removed)
-				{
-					Online::GetPartyInterfaceChecked(GetWorld())->RemoveUserForRejoin(*OwningLocalUserId, PartyId, MemberId);
-				}
-				else
-				{
-					// TODO:  Add a timer to remove players eventually
-				}
 			}
 		}
 		else
diff --git a/Engine/Plugins/Online/OnlineSubsystem/Source/Public/Interfaces/OnlinePartyInterface.h b/Engine/Plugins/Online/OnlineSubsystem/Source/Public/Interfaces/OnlinePartyInterface.h
--- a/Engine/Plugins/Online/OnlineSubsystem/Source/Public/Interfaces/OnlinePartyInterface.h
+++ b/Engine/Plugins/Online/OnlineSubsystem/Source/Public/Interfaces/OnlinePartyInterface.h
@@ -992,7 +992,8 @@ public:
 	 * @param PartyId - party the player is approved to rejoin the party
 	 * @param ApprovedUserId - the user that has been approved to attempt to rejoin the party (does not need to be in the party now)
 	 */
-	virtual void ApproveUserForRejoin(const FUniqueNetId& LocalUserId, const FOnlinePartyId& PartyId, const FUniqueNetId& ApprovedUserId) = 0;
+	UE_DEPRECATED(4.23, "Marking users for rejoins in the public interface is deprecated. This functionality should be implemented internal to your party implementation.")
+	virtual void ApproveUserForRejoin(const FUniqueNetId& LocalUserId, const FOnlinePartyId& PartyId, const FUniqueNetId& ApprovedUserId) {}
 
 	/** 
 	 * Unmark a user as approved to attempt to rejoin our party
@@ -1001,7 +1002,8 @@ public:
 	 * @param PartyId - party the player is approved to rejoin the party
 	 * @param RemovedUserId - the user that has lost approval to attempt to rejoin the party
 	 */
-	virtual void RemoveUserForRejoin(const FUniqueNetId& LocalUserId, const FOnlinePartyId& PartyId, const FUniqueNetId& RemovedUserId) = 0;
+	UE_DEPRECATED(4.23, "Marking users for rejoins in the public interface is deprecated. This functionality should be implemented internal to your party implementation.")
+	virtual void RemoveUserForRejoin(const FUniqueNetId& LocalUserId, const FOnlinePartyId& PartyId, const FUniqueNetId& RemovedUserId) {}
 
 	/** 
 	 * Get a list of users that have been approved for rejoining
@@ -1010,7 +1012,8 @@ public:
 	 * @param PartyId - party the player is approved to rejoin the party
 	 * @param OutApprovedUserIds - list of users that have been approved
 	 */
-	virtual void GetUsersApprovedForRejoin(const FUniqueNetId& LocalUserId, const FOnlinePartyId& PartyId, TArray<TSharedRef<const FUniqueNetId>>& OutApprovedUserIds) = 0;
+	UE_DEPRECATED(4.23, "Marking users for rejoins in the public interface is deprecated. This functionality should be implemented internal to your party implementation.")
+	virtual void GetUsersApprovedForRejoin(const FUniqueNetId& LocalUserId, const FOnlinePartyId& PartyId, TArray<TSharedRef<const FUniqueNetId>>& OutApprovedUserIds) {}
 
 	/**
 	 * Kick a user from an existing party
