Modify some IOnlinePartyJoinInfo accessor methods to return TSharedRef instead of const TSharedRef&

diff --git a/Engine/Plugins/Online/OnlineSubsystem/Source/Public/Interfaces/OnlinePartyInterface.h b/Engine/Plugins/Online/OnlineSubsystem/Source/Public/Interfaces/OnlinePartyInterface.h
--- a/Engine/Plugins/Online/OnlineSubsystem/Source/Public/Interfaces/OnlinePartyInterface.h
+++ b/Engine/Plugins/Online/OnlineSubsystem/Source/Public/Interfaces/OnlinePartyInterface.h
@@ -282,17 +282,17 @@ public:
 	/**
 	 * @return party id of party associated with this join invite
 	 */
-	virtual const TSharedRef<const FOnlinePartyId>& GetPartyId() const = 0;
+	virtual TSharedRef<const FOnlinePartyId> GetPartyId() const = 0;
 
 	/**
 	 * @return party id of party associated with this join invite
 	 */
-	virtual const FOnlinePartyTypeId GetPartyTypeId() const = 0;
+	virtual FOnlinePartyTypeId GetPartyTypeId() const = 0;
 
 	/**
 	 * @return user id of where this join info came from
 	 */
-	virtual const TSharedRef<const FUniqueNetId>& GetSourceUserId() const = 0;
+	virtual TSharedRef<const FUniqueNetId> GetSourceUserId() const = 0;
 
 	/**
 	 * @return user id of where this join info came from
