Add TOnlinePartyIdMap

diff --git a/Engine/Plugins/Online/OnlineSubsystem/Source/Public/OnlineSubsystemTypes.h b/Engine/Plugins/Online/OnlineSubsystem/Source/Public/OnlineSubsystemTypes.h
--- a/Engine/Plugins/Online/OnlineSubsystem/Source/Public/OnlineSubsystemTypes.h
+++ b/Engine/Plugins/Online/OnlineSubsystem/Source/Public/OnlineSubsystemTypes.h
@@ -1470,10 +1470,26 @@ protected:
 		return *this;
 	}
 
+	friend inline uint32 GetTypeHash(const FOnlinePartyId& Value)
+	{
+		return CityHash32(reinterpret_cast<const char*>(Value.GetBytes()), Value.GetSize());
+	}
+
 public:
 	virtual ~FOnlinePartyId() {}
 };
 
+template <typename ValueType>
+struct TOnlinePartyIdMapKeyFuncs : public TDefaultMapKeyFuncs<TSharedRef<const FOnlinePartyId>, ValueType, false>
+{
+	static FORCEINLINE TSharedRef<const FOnlinePartyId>	GetSetKey(TPair<TSharedRef<const FOnlinePartyId>, ValueType> const& Element) { return Element.Key; }
+	static FORCEINLINE uint32							GetKeyHash(TSharedRef<const FOnlinePartyId> const& Key) {	return GetTypeHash(*Key); }
+	static FORCEINLINE bool								Matches(TSharedRef<const FOnlinePartyId> const& A, TSharedRef<const FOnlinePartyId> const& B) { return (A == B) || (*A == *B); }
+};
+
+template <typename ValueType>
+using TOnlinePartyIdMap = TMap<TSharedRef<const FOnlinePartyId>, ValueType, FDefaultSetAllocator, TOnlinePartyIdMapKeyFuncs<ValueType>>;
+
 /**
  * Id of a party's type
  */
