Add FPartyInvitationRecipient::ToDebugString

diff --git a/Engine/Plugins/Online/OnlineSubsystem/Source/Private/OnlinePartyInterface.cpp b/Engine/Plugins/Online/OnlineSubsystem/Source/Private/OnlinePartyInterface.cpp
--- a/Engine/Plugins/Online/OnlineSubsystem/Source/Private/OnlinePartyInterface.cpp
+++ b/Engine/Plugins/Online/OnlineSubsystem/Source/Private/OnlinePartyInterface.cpp
@@ -130,6 +130,11 @@ void FOnlinePartyData::FromJson(const FString& JsonString)
 	}
 }
 
+FString FPartyInvitationRecipient::ToDebugString() const
+{
+	return FString::Printf(TEXT("Id=[%s], PlatformData=[%s]"), *Id->ToDebugString(), *PlatformData);
+}
+
 bool FPartyConfiguration::operator==(const FPartyConfiguration& Other) const
 {
 	return JoinRequestAction == Other.JoinRequestAction &&
diff --git a/Engine/Plugins/Online/OnlineSubsystem/Source/Public/Interfaces/OnlinePartyInterface.h b/Engine/Plugins/Online/OnlineSubsystem/Source/Public/Interfaces/OnlinePartyInterface.h
--- a/Engine/Plugins/Online/OnlineSubsystem/Source/Public/Interfaces/OnlinePartyInterface.h
+++ b/Engine/Plugins/Online/OnlineSubsystem/Source/Public/Interfaces/OnlinePartyInterface.h
@@ -491,19 +491,26 @@ enum class EMemberExitedReason
 	Kicked
 };
 
+/** Recipient information for SendInvitation */
 struct FPartyInvitationRecipient
 {
+	/** Constructor */
 	FPartyInvitationRecipient(const TSharedRef<const FUniqueNetId>& InId)
 		: Id(InId)
-		{}
+	{}
 
+	/** Constructor */
 	FPartyInvitationRecipient(const FUniqueNetId& InId)
 		: Id(InId.AsShared())
 	{}
 
-
+	/** Id of the user to send the invitation to */
 	TSharedRef<const FUniqueNetId> Id;
+	/** Additional data to provide context for the invitee */
 	FString PlatformData;
+
+	/** Get a string representation suitable for logging */
+	FString ONLINESUBSYSTEM_API ToDebugString() const;
 };
 
 
