Add delegate parameter names to IOnlineFriends

diff --git a/Engine/Plugins/Online/OnlineSubsystem/Source/Public/Interfaces/OnlineFriendsInterface.h b/Engine/Plugins/Online/OnlineSubsystem/Source/Public/Interfaces/OnlineFriendsInterface.h
--- a/Engine/Plugins/Online/OnlineSubsystem/Source/Public/Interfaces/OnlineFriendsInterface.h
+++ b/Engine/Plugins/Online/OnlineSubsystem/Source/Public/Interfaces/OnlineFriendsInterface.h
@@ -119,7 +119,7 @@ typedef FOnOutgoingInviteSent::FDelegate FOnOutgoingInviteSentDelegate;
  * @param ListName name of the friends list that was operated on
  * @param ErrorStr string representing the error condition
  */
-DECLARE_DELEGATE_FourParams(FOnReadFriendsListComplete, int32, bool, const FString&, const FString&);
+DECLARE_DELEGATE_FourParams(FOnReadFriendsListComplete, int32 /*LocalUserNum*/, bool /*bWasSuccessful*/, const FString& /*ListName*/, const FString& /*ErrorStr*/);
 
 /**
  * Delegate used when the friends list delete request has completed
@@ -129,7 +129,7 @@ DECLARE_DELEGATE_FourParams(FOnReadFriendsListComplete, int32, bool, const FStri
  * @param ListName name of the friends list that was operated on
  * @param ErrorStr string representing the error condition
  */
-DECLARE_DELEGATE_FourParams(FOnDeleteFriendsListComplete, int32, bool, const FString&, const FString&);
+DECLARE_DELEGATE_FourParams(FOnDeleteFriendsListComplete, int32 /*LocalUserNum*/, bool /*bWasSuccessful*/, const FString& /*ListName*/, const FString& /*ErrorStr*/);
 
 /**
  * Delegate used when an invite send request has completed
@@ -140,7 +140,7 @@ DECLARE_DELEGATE_FourParams(FOnDeleteFriendsListComplete, int32, bool, const FSt
  * @param ListName name of the friends list that was operated on
  * @param ErrorStr string representing the error condition
  */
-DECLARE_DELEGATE_FiveParams(FOnSendInviteComplete, int32, bool, const FUniqueNetId&, const FString&, const FString&);
+DECLARE_DELEGATE_FiveParams(FOnSendInviteComplete, int32 /*LocalUserNum*/, bool /*bWasSuccessful*/, const FUniqueNetId& /*FriendId*/, const FString& /*ListName*/, const FString& /*ErrorStr*/);
 
 /**
  * Delegate used when an invite accept request has completed
@@ -151,7 +151,7 @@ DECLARE_DELEGATE_FiveParams(FOnSendInviteComplete, int32, bool, const FUniqueNet
  * @param ListName name of the friends list that was operated on
  * @param ErrorStr string representing the error condition
  */
-DECLARE_DELEGATE_FiveParams(FOnAcceptInviteComplete, int32, bool, const FUniqueNetId&, const FString&, const FString&);
+DECLARE_DELEGATE_FiveParams(FOnAcceptInviteComplete, int32 /*LocalUserNum*/, bool /*bWasSuccessful*/, const FUniqueNetId& /*FriendId*/, const FString& /*ListName*/, const FString& /*ErrorStr*/);
 
 /**
  * Delegate used when an invite reject request has completed
@@ -162,7 +162,7 @@ DECLARE_DELEGATE_FiveParams(FOnAcceptInviteComplete, int32, bool, const FUniqueN
  * @param ListName name of the friends list that was operated on
  * @param ErrorStr string representing the error condition
  */
-DECLARE_MULTICAST_DELEGATE_FiveParams(FOnRejectInviteComplete, int32, bool, const FUniqueNetId&, const FString&, const FString&);
+DECLARE_MULTICAST_DELEGATE_FiveParams(FOnRejectInviteComplete, int32 /*LocalUserNum*/, bool /*bWasSuccessful*/, const FUniqueNetId& /*FriendId*/, const FString& /*ListName*/, const FString& /*ErrorStr*/);
 typedef FOnRejectInviteComplete::FDelegate FOnRejectInviteCompleteDelegate;
 
 /**
@@ -174,7 +174,7 @@ typedef FOnRejectInviteComplete::FDelegate FOnRejectInviteCompleteDelegate;
  * @param ListName name of the friends list that was operated on
  * @param ErrorStr string representing the error condition
  */
-DECLARE_MULTICAST_DELEGATE_FiveParams(FOnDeleteFriendComplete, int32, bool, const FUniqueNetId&, const FString&, const FString&);
+DECLARE_MULTICAST_DELEGATE_FiveParams(FOnDeleteFriendComplete, int32 /*LocalUserNum*/, bool /*bWasSuccessful*/, const FUniqueNetId& /*FriendId*/, const FString& /*ListName*/, const FString& /*ErrorStr*/);
 typedef FOnDeleteFriendComplete::FDelegate FOnDeleteFriendCompleteDelegate;
 
 /**
@@ -186,7 +186,7 @@ typedef FOnDeleteFriendComplete::FDelegate FOnDeleteFriendCompleteDelegate;
  * @param ListName name of the friends list that was operated on
  * @param ErrorStr string representing the error condition
  */
-DECLARE_MULTICAST_DELEGATE_FiveParams(FOnBlockedPlayerComplete, int32, bool, const FUniqueNetId&, const FString&, const FString&);
+DECLARE_MULTICAST_DELEGATE_FiveParams(FOnBlockedPlayerComplete, int32 /*LocalUserNum*/, bool /*bWasSuccessful*/, const FUniqueNetId& /*UniqueID*/, const FString& /*ListName*/, const FString& /*ErrorStr*/);
 typedef FOnBlockedPlayerComplete::FDelegate FOnBlockedPlayerCompleteDelegate;
 
 /**
@@ -198,7 +198,7 @@ typedef FOnBlockedPlayerComplete::FDelegate FOnBlockedPlayerCompleteDelegate;
  * @param ListName name of the friends list that was operated on
  * @param ErrorStr string representing the error condition
  */
-DECLARE_MULTICAST_DELEGATE_FiveParams(FOnUnblockedPlayerComplete, int32, bool, const FUniqueNetId&, const FString&, const FString&);
+DECLARE_MULTICAST_DELEGATE_FiveParams(FOnUnblockedPlayerComplete, int32 /*LocalUserNum*/, bool /*bWasSuccessful*/, const FUniqueNetId& /*UniqueID*/, const FString& /*ListName*/, const FString& /*ErrorStr*/);
 typedef FOnUnblockedPlayerComplete::FDelegate FOnUnblockedPlayerCompleteDelegate;
 
 /**
@@ -320,7 +320,7 @@ typedef FOnFriendRemoved::FDelegate FOnFriendRemovedDelegate;
 * @param Settings settings retrieved / updated
 * @param ErrorStr string representing the error condition
 */
-DECLARE_DELEGATE_FiveParams(FOnSettingsOperationComplete, const FUniqueNetId&, bool, bool, const FFriendSettings&, const FString&);
+DECLARE_DELEGATE_FiveParams(FOnSettingsOperationComplete, const FUniqueNetId& /*UserId*/, bool /*bWasSuccessful*/, bool /*bWasUpdate*/, const FFriendSettings& /*Settings*/, const FString& /*ErrorStr*/);
 
 /**
  * Interface definition for the online services friends services 
@@ -446,7 +446,7 @@ public:
 	 * @param ListName name of the friends list that was operated on
 	 * @param ErrorStr string representing the error condition
 	 */
-	DEFINE_ONLINE_PLAYER_DELEGATE_FOUR_PARAM(MAX_LOCAL_PLAYERS, OnRejectInviteComplete, bool, const FUniqueNetId&, const FString&, const FString&);
+	DEFINE_ONLINE_PLAYER_DELEGATE_FOUR_PARAM(MAX_LOCAL_PLAYERS, OnRejectInviteComplete, bool /*bWasSuccessful*/, const FUniqueNetId& /*FriendId*/, const FString& /*ListName*/, const FString& /*ErrorStr*/);
 
 	/**
 	 * Starts an async task that deletes a friend from the named friends list
@@ -468,29 +468,29 @@ public:
 	 * @param ListName name of the friends list that was operated on
 	 * @param ErrorStr string representing the error condition
 	 */
-	DEFINE_ONLINE_PLAYER_DELEGATE_FOUR_PARAM(MAX_LOCAL_PLAYERS, OnDeleteFriendComplete, bool, const FUniqueNetId&, const FString&, const FString&);
+	DEFINE_ONLINE_PLAYER_DELEGATE_FOUR_PARAM(MAX_LOCAL_PLAYERS, OnDeleteFriendComplete, bool /*bWasSuccessful*/, const FUniqueNetId& /*FriendId*/, const FString& /*ListName*/, const FString& /*ErrorStr*/);
 
 	/**
 	 * Delegate used when a block player request has completed
 	 *
 	 * @param LocalUserNum the controller number of the associated user that made the request
 	 * @param bWasSuccessful true if the async action completed without error, false if there was an error
-	 * @param UniqueIF player that was blocked
+	 * @param UniqueId player that was blocked
 	 * @param ListName name of the friends list that was operated on
 	 * @param ErrorStr string representing the error condition
 	 */
-	DEFINE_ONLINE_PLAYER_DELEGATE_FOUR_PARAM(MAX_LOCAL_PLAYERS, OnBlockedPlayerComplete, bool, const FUniqueNetId&, const FString&, const FString&);
+	DEFINE_ONLINE_PLAYER_DELEGATE_FOUR_PARAM(MAX_LOCAL_PLAYERS, OnBlockedPlayerComplete, bool /*bWasSuccessful*/, const FUniqueNetId& /*UniqueId*/, const FString& /*ListName*/, const FString& /*ErrorStr*/);
 
 	/**
-	 * Delegate used when an ublock player request has completed
+	 * Delegate used when an unblock player request has completed
 	 *
 	 * @param LocalUserNum the controller number of the associated user that made the request
 	 * @param bWasSuccessful true if the async action completed without error, false if there was an error
-	 * @param UniqueIF player that was unblocked
+	 * @param UniqueId player that was unblocked
 	 * @param ListName name of the friends list that was operated on
 	 * @param ErrorStr string representing the error condition
 	 */
-	DEFINE_ONLINE_PLAYER_DELEGATE_FOUR_PARAM(MAX_LOCAL_PLAYERS, OnUnblockedPlayerComplete, bool, const FUniqueNetId&, const FString&, const FString&);
+	DEFINE_ONLINE_PLAYER_DELEGATE_FOUR_PARAM(MAX_LOCAL_PLAYERS, OnUnblockedPlayerComplete, bool /*bWasSuccessful*/, const FUniqueNetId& /*UniqueId*/, const FString& /*ListName*/, const FString& /*ErrorStr*/);
 
 	/**
      * Delegate used in block list change notifications
