Expose FVariantData type suffix
Allow FVariantData::FromJsonValue to use a string value for a JSON variant type.

diff --git a/Engine/Plugins/Online/OnlineSubsystem/Source/Private/OnlineKeyValuePair.cpp b/Engine/Plugins/Online/OnlineSubsystem/Source/Private/OnlineKeyValuePair.cpp
--- a/Engine/Plugins/Online/OnlineSubsystem/Source/Private/OnlineKeyValuePair.cpp
+++ b/Engine/Plugins/Online/OnlineSubsystem/Source/Private/OnlineKeyValuePair.cpp
@@ -800,6 +800,55 @@ bool FVariantData::FromJson(const TSharedRef<FJsonObject>& JsonObject)
 	return bResult;
 }
 
+FString FVariantData::GetTypeSuffix() const
+{
+	switch (Type)
+	{
+	case EOnlineKeyValuePairDataType::Int32:
+	{
+		return TEXT("_i");
+	}
+	case EOnlineKeyValuePairDataType::UInt32:
+	{
+		return TEXT("_u");
+	}
+	case EOnlineKeyValuePairDataType::Float:
+	{
+		return TEXT("_f");
+	}
+	case EOnlineKeyValuePairDataType::String:
+	{
+		return TEXT("_s");
+	}
+	case EOnlineKeyValuePairDataType::Bool:
+	{
+		return TEXT("_b");
+	}
+	case EOnlineKeyValuePairDataType::Int64:
+	{
+		return TEXT("_I");
+	}
+	case EOnlineKeyValuePairDataType::UInt64:
+	{
+		return TEXT("_U");
+	}
+	case EOnlineKeyValuePairDataType::Double:
+	{
+		return TEXT("_d");
+	}
+	case EOnlineKeyValuePairDataType::Json:
+	{
+		return TEXT("_j");
+	}
+	case EOnlineKeyValuePairDataType::Empty:
+	case EOnlineKeyValuePairDataType::Blob:
+	default:
+		break;
+	}
+	// Default to no suffix
+	return TEXT("");
+}
+
 void FVariantData::AddToJsonObject(const TSharedRef<FJsonObject>& JsonObject, const FString& Name, const bool bWithTypeSuffix /* = true*/) const
 {
 	switch (Type)
@@ -962,11 +1011,24 @@ bool FVariantData::FromJsonValue(const FString& JsonPropertyName, const TSharedR
 		}
 		case TEXT('j'):
 		{
-			const TSharedPtr<FJsonObject>* FieldValue;
-			if (JsonValue->TryGetObject(FieldValue))
+			const TSharedPtr<FJsonObject>* FieldValueObject;
+			FString FieldValueString;
+			if (JsonValue->TryGetObject(FieldValueObject)) // Try to read as a JSON object
+			{
+				SetValue((*FieldValueObject).ToSharedRef());
+				bResult = true;
+			}
+			else if (JsonValue->TryGetString(FieldValueString)) // Try to read as a JSON string
 			{
-				SetValue((*FieldValue).ToSharedRef());
+				SetJsonValueFromString(FieldValueString);
+#if !UE_BUILD_SHIPPING
+				// See if this was a valid JSON string
+				TSharedPtr<FJsonObject> ValidJsonObject;
+				GetValue(ValidJsonObject);
+				bResult = ValidJsonObject.IsValid();
+#else
 				bResult = true;
+#endif
 			}
 		}
 		default:
diff --git a/Engine/Plugins/Online/OnlineSubsystem/Source/Public/OnlineKeyValuePair.h b/Engine/Plugins/Online/OnlineSubsystem/Source/Public/OnlineKeyValuePair.h
--- a/Engine/Plugins/Online/OnlineSubsystem/Source/Public/OnlineKeyValuePair.h
+++ b/Engine/Plugins/Online/OnlineSubsystem/Source/Public/OnlineKeyValuePair.h
@@ -425,6 +425,13 @@ public:
 	 */
 	bool FromJson(const TSharedRef<class FJsonObject>& JsonObject);
 
+	/**
+	 * Combine the type suffix given this variant's type
+	 *
+	 * @return the the type suffix
+	 */
+	FString GetTypeSuffix() const;
+	
 	/**
 	* Insert variant data into json object
 	*
