#UE4 - adding pin grant to Identity via EOS semantics
- EOS files excluded from this shelf

diff --git a/Engine/Plugins/Online/OnlineSubsystem/Source/Public/OnlineIdentityErrors.h b/Engine/Plugins/Online/OnlineSubsystem/Source/Public/OnlineIdentityErrors.h
--- a/Engine/Plugins/Online/OnlineSubsystem/Source/Public/OnlineIdentityErrors.h
+++ b/Engine/Plugins/Online/OnlineSubsystem/Source/Public/OnlineIdentityErrors.h
@@ -22,6 +22,9 @@ namespace OnlineIdentity
 
 		inline FOnlineError InvalidResult() { return ONLINE_ERROR(EOnlineErrorResult::InvalidResults, TEXT("invalid_result")); }
 
+		inline FOnlineError PinGrantFailure() { return ONLINE_ERROR(EOnlineErrorResult::FailExtended, TEXT("pin_grant_failure")); }
+		inline FOnlineError PinGrantTimeout() { return ONLINE_ERROR(EOnlineErrorResult::FailExtended, TEXT("pin_grant_timeout")); }
+
 		// Params
 		extern ONLINESUBSYSTEM_API const FString AuthLoginParam;
 		extern ONLINESUBSYSTEM_API const FString AuthTypeParam;
