Add FOnlinePartyData::RemoveAttribute
Fix typo in GetDirtyKeyVaAttrs, and change signature to give you removed attributes as well

diff --git a/Engine/Plugins/Online/OnlineSubsystem/Source/Private/OnlinePartyInterface.cpp b/Engine/Plugins/Online/OnlineSubsystem/Source/Private/OnlinePartyInterface.cpp
--- a/Engine/Plugins/Online/OnlineSubsystem/Source/Private/OnlinePartyInterface.cpp
+++ b/Engine/Plugins/Online/OnlineSubsystem/Source/Private/OnlinePartyInterface.cpp
@@ -22,17 +22,20 @@ bool FOnlinePartyData::operator!=(const FOnlinePartyData& Other) const
 	return !operator==(Other);
 }
 
-FOnlineKeyValuePairs<FString, FVariantData> FOnlinePartyData::GetDirtyKeyVaAttrs() const
+void FOnlinePartyData::GetDirtyKeyValAttrs(FOnlineKeyValuePairs<FString, FVariantData>& OutDirtyAttrs, TArray<FString>& OutRemovedAttrs) const
 {
-	FOnlineKeyValuePairs<FString, FVariantData> DirtyKeyValAttrs;
 	for (const FString& PropertyName : DirtyKeys)
 	{
 		const FVariantData* PropertyValue = KeyValAttrs.Find(PropertyName);
-		check(PropertyValue);
-
-		DirtyKeyValAttrs.Emplace(PropertyName, *PropertyValue);
+		if (PropertyValue)
+		{
+			OutDirtyAttrs.Emplace(PropertyName, *PropertyValue);
+		}
+		else
+		{
+			OutRemovedAttrs.Emplace(PropertyName);
+		}
 	}
-	return DirtyKeyValAttrs;
 }
 
 void FOnlinePartyData::ToJsonFull(FString& JsonString) const
diff --git a/Engine/Plugins/Online/OnlineSubsystem/Source/Public/Interfaces/OnlinePartyInterface.h b/Engine/Plugins/Online/OnlineSubsystem/Source/Public/Interfaces/OnlinePartyInterface.h
--- a/Engine/Plugins/Online/OnlineSubsystem/Source/Public/Interfaces/OnlinePartyInterface.h
+++ b/Engine/Plugins/Online/OnlineSubsystem/Source/Public/Interfaces/OnlinePartyInterface.h
@@ -109,7 +109,7 @@ public:
 		if (NewAttrValue != AttrValue)
 		{
 			NewAttrValue = AttrValue;
-			DirtyKeys.Add(AttrName);
+			DirtyKeys.Emplace(AttrName);
 		}
 	}
 
@@ -129,6 +129,32 @@ public:
 		}
 	}
 
+	/**
+	 * Remove an attribute from the party data
+	 *
+	 * @param AttrName - key for the attribute
+	 */
+	void RemoveAttribute(FString&& AttrName)
+	{
+		if (KeyValAttrs.Remove(AttrName) > 0)
+		{
+			DirtyKeys.Emplace(MoveTemp(AttrName));
+		}
+	}
+
+	/**
+	 * Remove an attribute from the party data
+	 *
+	 * @param AttrName - key for the attribute
+	 */
+	void RemoveAttribute(const FString& AttrName)
+	{
+		if (KeyValAttrs.Remove(AttrName) > 0)
+		{
+			DirtyKeys.Emplace(AttrName);
+		}
+	}
+
 	/**
 	 * Mark an attribute as dirty so it can be rebroadcasted
 	 *
@@ -150,11 +176,12 @@ public:
 	}
 
 	/**
-	 * Get the dirty key-value attributes
+	 * Get the dirty and removed key-value attributes
 	 *
-	 * @return the dirty key-value attributes
+	 * @param OutDirtyAttrs the dirty attributes
+	 * @param OutRemovedAttrs the removed attributes
 	 */
-	FOnlineKeyValuePairs<FString, FVariantData> GetDirtyKeyVaAttrs() const;
+	void GetDirtyKeyValAttrs(FOnlineKeyValuePairs<FString, FVariantData>& OutDirtyAttrs, TArray<FString>& OutRemovedAttrs) const;
 
 	/**
 	 * Clear the attributes map
@@ -2177,6 +2204,25 @@ inline FString ToDebugString(const IOnlinePartyJoinInfo& JoinInfo)
 		);
 }
 
+/**
+ * Dump key/value pairs for debugging
+ */
+inline FString ToDebugString(const FOnlineKeyValuePairs<FString, FVariantData>& KeyValAttrs)
+{
+	FString Result;
+	bool bPrintedAny = false;
+	for (const TPair<FString, FVariantData>& Iterator : KeyValAttrs)
+	{
+		if (bPrintedAny)
+		{
+			Result += TEXT(",");
+		}
+		Result += FString::Printf(TEXT("[%s=%s]"), *Iterator.Key, *Iterator.Value.ToString());
+		bPrintedAny = true;
+	}
+	return Result;
+}
+
 /**
 * Dump state about the party data for debugging
 */
@@ -2187,14 +2233,11 @@ inline FString ToDebugString(const FOnlinePartyData& PartyData)
 	int32 TotalBytesPerSec = PartyData.TotalPackets ? (PartyData.TotalBytes / PartyData.TotalPackets) : 0;
 	int32 TotalEffectiveBytesPerSec = PartyData.TotalPackets ? (PartyData.TotalEffectiveBytes / PartyData.TotalPackets) : 0;
 
-	Result += FString::Printf(TEXT("%dB [%d B/pkt], %dB [%d B/pkt], Rev: %d"), 
+	Result += FString::Printf(TEXT("%dB [%d B/pkt], %dB [%d B/pkt], Rev: %d,"), 
 		PartyData.TotalBytes, TotalBytesPerSec,
 		PartyData.TotalEffectiveBytes, TotalEffectiveBytesPerSec, 
 		PartyData.RevisionCount);
 
-	for (auto Iterator : PartyData.GetKeyValAttrs())
-	{
-		Result += FString::Printf(TEXT(",[%s=%s]"), *(Iterator.Key), *(Iterator.Value.ToString()));
-	}
+	Result += ToDebugString(PartyData.GetKeyValAttrs());
 	return Result;
 }
