Expose more party join request parameters on IOnlinePartyPendingJoinRequestInfo

diff --git a/Engine/Plugins/Online/OnlineSubsystem/Source/Public/Interfaces/OnlinePartyInterface.h b/Engine/Plugins/Online/OnlineSubsystem/Source/Public/Interfaces/OnlinePartyInterface.h
--- a/Engine/Plugins/Online/OnlineSubsystem/Source/Public/Interfaces/OnlinePartyInterface.h
+++ b/Engine/Plugins/Online/OnlineSubsystem/Source/Public/Interfaces/OnlinePartyInterface.h
@@ -284,14 +284,24 @@ public:
 	virtual ~IOnlinePartyPendingJoinRequestInfo() {}
 
 	/**
-	* @return id of the sender of this join request
-	*/
+	 * @return id of the sender of this join request
+	 */
 	virtual const TSharedRef<const FUniqueNetId>& GetSenderId() const = 0;
 
 	/**
-	* @return display name of the sender of this join request
-	*/
+	 * @return display name of the sender of this join request
+	 */
 	virtual const FString& GetSenderDisplayName() const = 0;
+
+	/**
+	 * @return platform of the sender of this join request
+	 */
+	virtual const FString& GetSenderPlatform() const = 0;
+
+	/**
+	 * @return join data provided by the sender for htis join request
+	 */
+	virtual TSharedRef<const FOnlinePartyData> GetSenderJoinData() const = 0;
 };
 
 /**
