Added new argument "PreviousPartyState" to the IOnlinePartySystem::OnPartyStateChanged delegate.
Added new enumeration "MemberConnectionStatus" to FOnlinePartyMember which can be used by the OSS to describe the member's connection to the local user.
Added new Delegate "OnMemberConnectionStatusChanged" that is fired whenever MemberConnectionStatus is changed.

diff --git a/Engine/Plugins/Online/OnlineSubsystem/Source/Public/Interfaces/OnlinePartyInterface.h b/Engine/Plugins/Online/OnlineSubsystem/Source/Public/Interfaces/OnlinePartyInterface.h
--- a/Engine/Plugins/Online/OnlineSubsystem/Source/Public/Interfaces/OnlinePartyInterface.h
+++ b/Engine/Plugins/Online/OnlineSubsystem/Source/Public/Interfaces/OnlinePartyInterface.h
@@ -38,12 +38,45 @@ enum class EUpdateConfigCompletionResult;
 
 struct FAnalyticsEventAttribute;
 
+enum class EMemberConnectionStatus
+{
+	Uninitialized,
+	Disconnected,
+	Initializing,
+	Connected
+};
+
+/**
+* notification when a party member's data status has changed
+* @param ChangedUserId - id associated with this notification
+* @param NewMemberConnectionStatus - new member data status
+*/
+DECLARE_MULTICAST_DELEGATE_ThreeParams(F_PREFIX(OnMemberConnectionStatusChanged), const FUniqueNetId& /*ChangedUserId*/, const EMemberConnectionStatus /*NewMemberDataStatus*/, const EMemberConnectionStatus /*PreviousMemberDataStatus*/);
+PARTY_DECLARE_DELEGATETYPE(OnMemberConnectionStatusChanged);
+
 /**
  * Party member user info returned by IOnlineParty interface
  */
 class FOnlinePartyMember : public FOnlineUser
 {
 public:
+	EMemberConnectionStatus MemberConnectionStatus = EMemberConnectionStatus::Uninitialized;
+	EMemberConnectionStatus PreviousMemberConnectionStatus = EMemberConnectionStatus::Uninitialized;
+
+	/**
+	* notification when a party member's data status has changed
+	* @param ChangedUserId - id associated with this notification
+	* @param NewMemberConnectionStatus - new member data status
+	* @param PreviousMemberConnectionStatus - previous member data status
+	*/
+	DEFINE_ONLINE_DELEGATE_THREE_PARAM(OnMemberConnectionStatusChanged, const FUniqueNetId& /*ChangedUserId*/, const EMemberConnectionStatus /*NewMemberConnectionStatus*/, const EMemberConnectionStatus /*PreviousMemberConnectionStatus*/);
+
+	void SetMemberConnectionStatus(EMemberConnectionStatus NewMemberConnectionStatus)
+	{
+		PreviousMemberConnectionStatus = MemberConnectionStatus;
+		MemberConnectionStatus = NewMemberConnectionStatus;
+		TriggerOnMemberConnectionStatusChangedDelegates(*GetUserId(), MemberConnectionStatus, PreviousMemberConnectionStatus);
+	}
 };
 
 /**
@@ -272,6 +305,7 @@ private:
 
 	/** set of which fields are dirty and need to transmitted */
 	TSet<FString> DirtyKeys;
+
 };
 
 /**
@@ -497,6 +531,7 @@ protected:
 		: PartyId(InPartyId)
 		, PartyTypeId(InPartyTypeId)
 		, State(EPartyState::None)
+		, PreviousState(EPartyState::None)
 		, Config(InConfig)
 	{}
 
@@ -517,6 +552,11 @@ public:
 	 * @return true if this party is joinable, false if not
 	 */
 	virtual bool IsJoinable() const = 0;
+	virtual void SetState(EPartyState InState)
+	{
+		PreviousState = State;
+		State = InState;
+	}
 
 	/**
 	 * Get the party's configuration
@@ -533,6 +573,8 @@ public:
 	TSharedPtr<const FUniqueNetId> LeaderId;
 	/** The current state of the party */
 	EPartyState State;
+	/** The current state of the party */
+	EPartyState PreviousState;
 	/** Current state of configuration */
 	TSharedRef<FPartyConfiguration> Config;
 	/** id of chat room associated with the party */
@@ -710,8 +752,9 @@ PARTY_DECLARE_DELEGATETYPE(OnPartyExited);
  * @param LocalUserId - id associated with this notification
  * @param PartyId - id associated with the party
  * @param State - state of the party
+ * @param PreviousState - previous state of the party
  */
-DECLARE_MULTICAST_DELEGATE_ThreeParams(F_PREFIX(OnPartyStateChanged), const FUniqueNetId& /*LocalUserId*/, const FOnlinePartyId& /*PartyId*/, EPartyState /*State*/);
+DECLARE_MULTICAST_DELEGATE_FourParams(F_PREFIX(OnPartyStateChanged), const FUniqueNetId& /*LocalUserId*/, const FOnlinePartyId& /*PartyId*/, EPartyState /*State*/, EPartyState /*PreviousState*/);
 PARTY_DECLARE_DELEGATETYPE(OnPartyStateChanged);
 
 /**
@@ -1364,8 +1407,9 @@ public:
 	 * @param LocalUserId - id associated with this notification
 	 * @param PartyId - id associated with the party
 	 * @param State - state of the party
+	 * @param PreviousState - previous state of the party
 	*/
-	DEFINE_ONLINE_DELEGATE_THREE_PARAM(OnPartyStateChanged, const FUniqueNetId& /*LocalUserId*/, const FOnlinePartyId& /*PartyId*/, EPartyState /*State*/);
+	DEFINE_ONLINE_DELEGATE_FOUR_PARAM(OnPartyStateChanged, const FUniqueNetId& /*LocalUserId*/, const FOnlinePartyId& /*PartyId*/, EPartyState /*State*/, EPartyState /*PreviousState*/);
 
 	/**
 	* notification of when a player had been approved to Join In Progress
