Add IOnlinePartySystem::RestoreParties for restoring party memberships

diff --git a/Engine/Plugins/Online/OnlineSubsystem/Source/Public/Interfaces/OnlinePartyInterface.h b/Engine/Plugins/Online/OnlineSubsystem/Source/Public/Interfaces/OnlinePartyInterface.h
--- a/Engine/Plugins/Online/OnlineSubsystem/Source/Public/Interfaces/OnlinePartyInterface.h
+++ b/Engine/Plugins/Online/OnlineSubsystem/Source/Public/Interfaces/OnlinePartyInterface.h
@@ -611,7 +611,7 @@ struct FPartyInvitationRecipient
 	FString ONLINESUBSYSTEM_API ToDebugString() const;
 };
 
-
+struct FOnlineError;
 enum class ECreatePartyCompletionResult;
 enum class EJoinPartyCompletionResult;
 enum class ELeavePartyCompletionResult;
@@ -627,6 +627,13 @@ enum class EInvitationResponse;
 ///////////////////////////////////////////////////////////////////
 // Completion delegates
 ///////////////////////////////////////////////////////////////////
+/**
+ * Restore parties async task completed callback
+ *
+ * @param LocalUserId id of user that initiated the request
+ * @param Result Result of the operation
+ */
+DECLARE_DELEGATE_TwoParams(FOnRestorePartiesComplete, const FUniqueNetId& /*LocalUserId*/, const FOnlineError& /*Result*/);
 /**
  * Party creation async task completed callback
  *
@@ -929,6 +936,14 @@ protected:
 public:
 	virtual ~IOnlinePartySystem() {};
 
+	/**
+	 * Restore party memberships. Intended to be called once during login to restore state from other running instances.
+	 *
+	 * @param LocalUserId the user to restore the party membership for
+	 * @param CompletionDelegate the delegate to trigger on completion
+	 */
+	virtual void RestoreParties(const FUniqueNetId& LocalUserId, const FOnRestorePartiesComplete& CompletionDelegate) = 0;
+	
 	/**
 	 * Create a new party
 	 *
