#kairos

Only initialize embedded communication module when loading default OSS instance
This was causing cyclical intiialization when the embedded communication module was initialized and tried to load the MCP_SUBSYSTEM_EMBEDDED

diff --git a/Engine/Plugins/Online/OnlineSubsystem/Source/Private/OnlineSubsystemModule.cpp b/Engine/Plugins/Online/OnlineSubsystem/Source/Private/OnlineSubsystemModule.cpp
--- a/Engine/Plugins/Online/OnlineSubsystem/Source/Private/OnlineSubsystemModule.cpp
+++ b/Engine/Plugins/Online/OnlineSubsystem/Source/Private/OnlineSubsystemModule.cpp
@@ -311,7 +311,7 @@ IOnlineSubsystem* FOnlineSubsystemModule::GetOnlineSubsystem(const FName InSubsy
 
 				if (OSSFactory != nullptr)
 				{
-					UE_LOG_ONLINE(Verbose, TEXT("Creating online subsystem instance for: %s"), *InSubsystemName.ToString());
+					UE_LOG_ONLINE(Log, TEXT("Creating online subsystem instance for: %s"), *InSubsystemName.ToString());
 						
 					IOnlineSubsystemPtr NewSubsystemInstance = (*OSSFactory)->CreateSubsystem(InstanceName);
 					if (NewSubsystemInstance.IsValid())
