Fix default subsystem not getting updated in PIE when the user enables login

diff --git a/Engine/Plugins/Online/OnlineSubsystem/Source/Private/OnlineSubsystemModule.cpp b/Engine/Plugins/Online/OnlineSubsystem/Source/Private/OnlineSubsystemModule.cpp
--- a/Engine/Plugins/Online/OnlineSubsystem/Source/Private/OnlineSubsystemModule.cpp
+++ b/Engine/Plugins/Online/OnlineSubsystem/Source/Private/OnlineSubsystemModule.cpp
@@ -167,6 +167,8 @@ void FOnlineSubsystemModule::LoadDefaultSubsystem()
 void FOnlineSubsystemModule::ReloadDefaultSubsystem()
 {
 	DestroyOnlineSubsystem(DefaultPlatformService);
+	// Clear our InstanceNames cache so we can re-establish it in case the DefaultPlatformService
+	InstanceNames.Empty();
 	LoadDefaultSubsystem();
 }
 
@@ -233,14 +235,6 @@ void FOnlineSubsystemModule::EnumerateOnlineSubsystems(FEnumerateOnlineSubsystem
 
 FName FOnlineSubsystemModule::ParseOnlineSubsystemName(const FName& FullName, FName& SubsystemName, FName& InstanceName) const
 {
-	struct FInstanceNameEntry
-	{
-		FName SubsystemName;
-		FName InstanceName;
-		FName FullPath;
-	};
-	static TMap<FName, FInstanceNameEntry> InstanceNames;
-
 	FInstanceNameEntry* Entry = InstanceNames.Find(FullName);
 	if (Entry)
 	{
diff --git a/Engine/Plugins/Online/OnlineSubsystem/Source/Public/OnlineSubsystemModule.h b/Engine/Plugins/Online/OnlineSubsystem/Source/Public/OnlineSubsystemModule.h
--- a/Engine/Plugins/Online/OnlineSubsystem/Source/Public/OnlineSubsystemModule.h
+++ b/Engine/Plugins/Online/OnlineSubsystem/Source/Public/OnlineSubsystemModule.h
@@ -223,6 +223,16 @@ public:
 	{
 		return true;
 	}
+
+private:
+	// Cached instance names for efficient lookup in ParseOnlineSubsystemName
+	struct FInstanceNameEntry
+	{
+		FName SubsystemName;
+		FName InstanceName;
+		FName FullPath;
+	};
+	mutable TMap<FName, FInstanceNameEntry> InstanceNames;
 };
 
 /** Public references to the online subsystem module pointer should use this */
