We now only trigger OnMemberConnectionStatusChanged delegates if we detect the status has changed.

diff --git a/Engine/Plugins/Online/OnlineSubsystem/Source/Public/Interfaces/OnlinePartyInterface.h b/Engine/Plugins/Online/OnlineSubsystem/Source/Public/Interfaces/OnlinePartyInterface.h
--- a/Engine/Plugins/Online/OnlineSubsystem/Source/Public/Interfaces/OnlinePartyInterface.h
+++ b/Engine/Plugins/Online/OnlineSubsystem/Source/Public/Interfaces/OnlinePartyInterface.h
@@ -73,9 +73,12 @@ public:
 
 	void SetMemberConnectionStatus(EMemberConnectionStatus NewMemberConnectionStatus)
 	{
-		PreviousMemberConnectionStatus = MemberConnectionStatus;
-		MemberConnectionStatus = NewMemberConnectionStatus;
-		TriggerOnMemberConnectionStatusChangedDelegates(*GetUserId(), MemberConnectionStatus, PreviousMemberConnectionStatus);
+		if (NewMemberConnectionStatus != MemberConnectionStatus)
+		{
+			PreviousMemberConnectionStatus = MemberConnectionStatus;
+			MemberConnectionStatus = NewMemberConnectionStatus;
+			TriggerOnMemberConnectionStatusChangedDelegates(*GetUserId(), MemberConnectionStatus, PreviousMemberConnectionStatus);
+		}
 	}
 };
 
