changes to PartySystemStateChanged delegate

rob.cannaday

diff --git a/Engine/Plugins/Online/OnlineFramework/Source/Party/Private/Party/SocialParty.cpp b/Engine/Plugins/Online/OnlineFramework/Source/Party/Private/Party/SocialParty.cpp
--- a/Engine/Plugins/Online/OnlineFramework/Source/Party/Private/Party/SocialParty.cpp
+++ b/Engine/Plugins/Online/OnlineFramework/Source/Party/Private/Party/SocialParty.cpp
@@ -920,7 +920,7 @@ void USocialParty::HandlePartyMemberExited(const FUniqueNetId& LocalUserId, cons
 	}
 }
 
-void USocialParty::HandlePartySystemStateChange(const FUniqueNetId& LocalUserId, EPartySystemState NewState)
+void USocialParty::HandlePartySystemStateChange(EPartySystemState NewState)
 {
 	UE_LOG(LogParty, VeryVerbose, TEXT("Party [%s] received notification of a party system state change to [%d]"), *ToDebugString(), (int32)NewState);
 	if (NewState == EPartySystemState::RequestingShutdown)
diff --git a/Engine/Plugins/Online/OnlineFramework/Source/Party/Public/Party/SocialParty.h b/Engine/Plugins/Online/OnlineFramework/Source/Party/Public/Party/SocialParty.h
--- a/Engine/Plugins/Online/OnlineFramework/Source/Party/Public/Party/SocialParty.h
+++ b/Engine/Plugins/Online/OnlineFramework/Source/Party/Public/Party/SocialParty.h
@@ -315,7 +315,7 @@ private:	// Handlers
 
 	void HandleLeavePartyComplete(const FUniqueNetId& LocalUserId, const FOnlinePartyId& PartyId, ELeavePartyCompletionResult LeaveResult, FOnLeavePartyAttemptComplete OnAttemptComplete);
 	
-	void HandlePartySystemStateChange(const FUniqueNetId& LocalUserId, EPartySystemState NewState);
+	void HandlePartySystemStateChange(EPartySystemState NewState);
 
 private:
 	TSharedPtr<const FOnlineParty> OssParty;
diff --git a/Engine/Plugins/Online/OnlineSubsystem/Source/Public/Interfaces/OnlinePartyInterface.h b/Engine/Plugins/Online/OnlineSubsystem/Source/Public/Interfaces/OnlinePartyInterface.h
--- a/Engine/Plugins/Online/OnlineSubsystem/Source/Public/Interfaces/OnlinePartyInterface.h
+++ b/Engine/Plugins/Online/OnlineSubsystem/Source/Public/Interfaces/OnlinePartyInterface.h
@@ -937,10 +937,9 @@ PARTY_DECLARE_DELEGATETYPE(OnPartyAnalyticsEvent);
 
 /**
 * Notification of party system state change
-* @param LocalUserId - id associated with this event
 * @param NewState - new state this partysystem is in
 */
-DECLARE_MULTICAST_DELEGATE_TwoParams(F_PREFIX(OnPartySystemStateChange), const FUniqueNetId& /*LocalUserId*/, EPartySystemState /*NewState*/);
+DECLARE_MULTICAST_DELEGATE_OneParam(F_PREFIX(OnPartySystemStateChange), EPartySystemState /*NewState*/);
 PARTY_DECLARE_DELEGATETYPE(OnPartySystemStateChange);
 
 /**
@@ -1595,10 +1594,9 @@ public:
 
 	/**
 	* Notification of party system state change
-	* @param LocalUserId - id associated with this event
 	* @param NewState - new state this partysystem is in
 	*/
-	DEFINE_ONLINE_DELEGATE_TWO_PARAM(OnPartySystemStateChange, const FUniqueNetId& /*LocalUserId*/, EPartySystemState /*NewState*/);
+	DEFINE_ONLINE_DELEGATE_ONE_PARAM(OnPartySystemStateChange, EPartySystemState /*NewState*/);
 
 	/**
 	 * Dump out party state for all known parties
