Fix memory leak when storing Json in FVariantData

diff --git a/Engine/Plugins/Online/OnlineSubsystem/Source/Private/OnlineKeyValuePair.cpp b/Engine/Plugins/Online/OnlineSubsystem/Source/Private/OnlineKeyValuePair.cpp
--- a/Engine/Plugins/Online/OnlineSubsystem/Source/Private/OnlineKeyValuePair.cpp
+++ b/Engine/Plugins/Online/OnlineSubsystem/Source/Private/OnlineKeyValuePair.cpp
@@ -486,6 +486,7 @@ void FVariantData::Empty()
 	// Be sure to delete deep allocations
 	switch (Type)
 	{
+	case EOnlineKeyValuePairDataType::Json:
 	case EOnlineKeyValuePairDataType::String:
 		delete [] Value.AsTCHAR;
 		break;
