#online
Separate config section for enabling online subsystem by instance name

diff --git a/Engine/Plugins/Online/OnlineSubsystem/Source/Private/OnlineSubsystem.cpp b/Engine/Plugins/Online/OnlineSubsystem/Source/Private/OnlineSubsystem.cpp
--- a/Engine/Plugins/Online/OnlineSubsystem/Source/Private/OnlineSubsystem.cpp
+++ b/Engine/Plugins/Online/OnlineSubsystem/Source/Private/OnlineSubsystem.cpp
@@ -288,7 +288,7 @@ FAutoConsoleCommand CmdResetAchievements(
 
 #endif
 
-bool IOnlineSubsystem::IsEnabled(const FName& SubsystemName)
+bool IOnlineSubsystem::IsEnabled(const FName& SubsystemName, const FName& InstanceName)
 {
 	bool bIsDisabledByCommandLine = false;
 #if !UE_BUILD_SHIPPING && !UE_BUILD_SHIPPING_WITH_EDITOR
@@ -299,9 +299,27 @@ bool IOnlineSubsystem::IsEnabled(const FName& SubsystemName)
 	if (!bIsDisabledByCommandLine)
 	{
 		bool bIsEnabledByConfig = false;
-		const FString ConfigSection(FString::Printf(TEXT("OnlineSubsystem%s"), *SubsystemName.ToString()));
-		const bool bConfigOptionExists = GConfig->GetBool(*ConfigSection, TEXT("bEnabled"), bIsEnabledByConfig, GEngineIni);
-		UE_CLOG_ONLINE(!bConfigOptionExists, Verbose, TEXT("[%s].bEnabled is not set, defaulting to true"), *ConfigSection);
+		bool bConfigOptionExists = false;
+		for (int32 InstancePass = 0; InstancePass < 2; InstancePass++)
+		{
+			FString InstanceSection;
+			if (InstancePass == 1)
+			{
+				if (InstanceName != NAME_None)
+				{
+					InstanceSection = TEXT(" ") + InstanceName.ToString();
+				}
+				else
+				{
+					continue;
+				}
+			}
+
+			FString ConfigSection(FString::Printf(TEXT("OnlineSubsystem%s"), *SubsystemName.ToString()));
+			ConfigSection += InstanceSection;
+			bConfigOptionExists |= GConfig->GetBool(*ConfigSection, TEXT("bEnabled"), bIsEnabledByConfig, GEngineIni);
+		}
+		UE_CLOG_ONLINE(!bConfigOptionExists, Verbose, TEXT("[OnlineSubsystem%s].bEnabled is not set, defaulting to true"), *SubsystemName.ToString());
 	
 		return !bConfigOptionExists || bIsEnabledByConfig;
 	}
diff --git a/Engine/Plugins/Online/OnlineSubsystem/Source/Private/OnlineSubsystemImpl.cpp b/Engine/Plugins/Online/OnlineSubsystem/Source/Private/OnlineSubsystemImpl.cpp
--- a/Engine/Plugins/Online/OnlineSubsystem/Source/Private/OnlineSubsystemImpl.cpp
+++ b/Engine/Plugins/Online/OnlineSubsystem/Source/Private/OnlineSubsystemImpl.cpp
@@ -258,7 +258,7 @@ bool FOnlineSubsystemImpl::Exec(UWorld* InWorld, const TCHAR* Cmd, FOutputDevice
 
 bool FOnlineSubsystemImpl::IsEnabled() const
 {
-	return IOnlineSubsystem::IsEnabled(SubsystemName);
+	return IOnlineSubsystem::IsEnabled(SubsystemName, InstanceName == FOnlineSubsystemImpl::DefaultInstanceName ? NAME_None : InstanceName);
 }
 
 void FOnlineSubsystemImpl::DumpReceipts(const FUniqueNetId& UserId)
diff --git a/Engine/Plugins/Online/OnlineSubsystem/Source/Public/OnlineSubsystem.h b/Engine/Plugins/Online/OnlineSubsystem/Source/Public/OnlineSubsystem.h
--- a/Engine/Plugins/Online/OnlineSubsystem/Source/Public/OnlineSubsystem.h
+++ b/Engine/Plugins/Online/OnlineSubsystem/Source/Public/OnlineSubsystem.h
@@ -252,7 +252,7 @@ public:
 	 * @param SubsystemName - Name of the requested online service
 	 * @return true if the subsystem is enabled by config
 	 */
-	static bool IsEnabled(const FName& SubsystemName);
+	static bool IsEnabled(const FName& SubsystemName, const FName& InstanceName);
 
 	/**
 	 * Return the name of the subsystem @see OnlineSubsystemNames.h
