#online
Default param for CIS issues

diff --git a/Engine/Plugins/Online/OnlineSubsystem/Source/Public/OnlineSubsystem.h b/Engine/Plugins/Online/OnlineSubsystem/Source/Public/OnlineSubsystem.h
--- a/Engine/Plugins/Online/OnlineSubsystem/Source/Public/OnlineSubsystem.h
+++ b/Engine/Plugins/Online/OnlineSubsystem/Source/Public/OnlineSubsystem.h
@@ -252,7 +252,7 @@ public:
 	 * @param SubsystemName - Name of the requested online service
 	 * @return true if the subsystem is enabled by config
 	 */
-	static bool IsEnabled(const FName& SubsystemName, const FName& InstanceName);
+	static bool IsEnabled(const FName& SubsystemName, const FName& InstanceName = NAME_None);
 
 	/**
 	 * Return the name of the subsystem @see OnlineSubsystemNames.h
