Fix for crash when using LAN beacons with mixed versions including 4.22.

diff --git a/Engine/Plugins/Online/OnlineSubsystem/Source/Public/LANBeacon.h b/Engine/Plugins/Online/OnlineSubsystem/Source/Public/LANBeacon.h
--- a/Engine/Plugins/Online/OnlineSubsystem/Source/Public/LANBeacon.h
+++ b/Engine/Plugins/Online/OnlineSubsystem/Source/Public/LANBeacon.h
@@ -15,7 +15,7 @@
  *
  *	<Ver byte><Platform byte><Game unique 4 bytes><packet type 2 bytes><nonce 8 bytes><payload>
  */
-#define LAN_BEACON_PACKET_VERSION (uint8)10
+#define LAN_BEACON_PACKET_VERSION (uint8)11
 
 /** The size of the header for validation */
 #define LAN_BEACON_PACKET_HEADER_SIZE 16
