#online,mcp,friends

Fix for friend invites not being sent out for anyone already marked as a suggested friend (ie. via FB/Steam friends lists)

diff --git a/Engine/Plugins/Online/OnlineSubsystem/Source/Public/OnlineSubsystemTypes.h b/Engine/Plugins/Online/OnlineSubsystem/Source/Public/OnlineSubsystemTypes.h
--- a/Engine/Plugins/Online/OnlineSubsystem/Source/Public/OnlineSubsystemTypes.h
+++ b/Engine/Plugins/Online/OnlineSubsystem/Source/Public/OnlineSubsystemTypes.h
@@ -1244,6 +1244,8 @@ namespace EInviteStatus
 		PendingOutbound,
 		/** Player has been blocked */
 		Blocked,
+		/** Suggested friend */
+		Suggested
 	};
 
 	/** 
@@ -1273,6 +1275,10 @@ namespace EInviteStatus
 			{
 				return TEXT("Blocked");
 			}
+			case Suggested:
+			{
+				return TEXT("Suggested");
+			}
 		}
 		return TEXT("");
 	}
