Make accessors for party system const
Make party enums uint8, or int8 where negative values are expected

diff --git a/Engine/Plugins/Online/OnlineFramework/Source/Party/Private/Party/PartyMember.cpp b/Engine/Plugins/Online/OnlineFramework/Source/Party/Private/Party/PartyMember.cpp
--- a/Engine/Plugins/Online/OnlineFramework/Source/Party/Private/Party/PartyMember.cpp
+++ b/Engine/Plugins/Online/OnlineFramework/Source/Party/Private/Party/PartyMember.cpp
@@ -61,7 +61,7 @@ void UPartyMember::BeginDestroy()
 	}
 }
 
-void UPartyMember::InitializePartyMember(const TSharedRef<FOnlinePartyMember>& InOssMember, const FSimpleDelegate& OnInitComplete)
+void UPartyMember::InitializePartyMember(const FOnlinePartyMemberConstRef& InOssMember, const FSimpleDelegate& OnInitComplete)
 {
 	checkf(MemberDataReplicator.IsValid(), TEXT("Child classes of UPartyMember MUST call MemberRepData.EstablishRepDataInstance with a valid FPartyMemberRepData struct instance in their constructor."));
 	MemberDataReplicator->SetOwningMember(*this);
@@ -69,7 +69,7 @@ void UPartyMember::InitializePartyMember(const TSharedRef<FOnlinePartyMember>& I
 	if (ensure(!OssPartyMember.IsValid()))
 	{
 		OssPartyMember = InOssMember;
-		OssPartyMember->OnMemberConnectionStatusChangedDelegates.AddUObject(this, &ThisClass::HandleMemberConnectionStatusChanged);
+		OssPartyMember->OnMemberConnectionStatusChanged().AddUObject(this, &ThisClass::HandleMemberConnectionStatusChanged);
 		USocialToolkit* OwnerToolkit = GetParty().GetSocialManager().GetFirstLocalUserToolkit();
 		check(OwnerToolkit);
 
@@ -183,12 +183,12 @@ bool UPartyMember::IsLocalPlayer() const
 	return GetPrimaryNetId() == GetParty().GetOwningLocalUserId();
 }
 
-void UPartyMember::NotifyMemberDataReceived(const TSharedRef<FOnlinePartyData>& MemberData)
+void UPartyMember::NotifyMemberDataReceived(const FOnlinePartyData& MemberData)
 {
 	UE_LOG(LogParty, VeryVerbose, TEXT("Received updated rep data for member [%s]"), *ToDebugString());
 
 	check(MemberDataReplicator.IsValid());
-	MemberDataReplicator.ProcessReceivedData(*MemberData/*, bHasReceivedInitialData*/);
+	MemberDataReplicator.ProcessReceivedData(MemberData/*, bHasReceivedInitialData*/);
 
 	if (!bHasReceivedInitialData)
 	{
diff --git a/Engine/Plugins/Online/OnlineFramework/Source/Party/Private/Party/PartyPlatformSessionMonitor.h b/Engine/Plugins/Online/OnlineFramework/Source/Party/Private/Party/PartyPlatformSessionMonitor.h
--- a/Engine/Plugins/Online/OnlineFramework/Source/Party/Private/Party/PartyPlatformSessionMonitor.h
+++ b/Engine/Plugins/Online/OnlineFramework/Source/Party/Private/Party/PartyPlatformSessionMonitor.h
@@ -14,7 +14,7 @@ class FPartyPlatformSessionMonitor;
 
 struct FPartyRepData;
 struct FPartyConfiguration;
-enum class EMemberExitedReason;
+enum class EMemberExitedReason : uint8;
 
 class FPartyPlatformSessionManager : public TSharedFromThis<FPartyPlatformSessionManager>
 {
diff --git a/Engine/Plugins/Online/OnlineFramework/Source/Party/Private/Party/SocialParty.cpp b/Engine/Plugins/Online/OnlineFramework/Source/Party/Private/Party/SocialParty.cpp
--- a/Engine/Plugins/Online/OnlineFramework/Source/Party/Private/Party/SocialParty.cpp
+++ b/Engine/Plugins/Online/OnlineFramework/Source/Party/Private/Party/SocialParty.cpp
@@ -419,8 +419,8 @@ void USocialParty::InitializeParty(const TSharedRef<const FOnlineParty>& InOssPa
 void USocialParty::InitializePartyInternal()
 {
 	IOnlinePartyPtr PartyInterface = Online::GetPartyInterfaceChecked(GetWorld());
-	PartyInterface->AddOnPartyConfigChangedDelegate_Handle(FOnPartyConfigChangedDelegate::CreateUObject(this, &USocialParty::HandlePartyConfigChanged));
-	PartyInterface->AddOnPartyDataReceivedDelegate_Handle(FOnPartyDataReceivedDelegate::CreateUObject(this, &USocialParty::HandlePartyDataReceived));
+	PartyInterface->AddOnPartyConfigChangedDelegate_Handle(FOnPartyConfigChangedConstDelegate::CreateUObject(this, &USocialParty::HandlePartyConfigChanged));
+	PartyInterface->AddOnPartyDataReceivedDelegate_Handle(FOnPartyDataReceivedConstDelegate::CreateUObject(this, &USocialParty::HandlePartyDataReceived));
 	PartyInterface->AddOnPartyJoinRequestReceivedDelegate_Handle(FOnPartyJoinRequestReceivedDelegate::CreateUObject(this, &USocialParty::HandlePartyJoinRequestReceived));
 	PartyInterface->AddOnPartyJIPRequestReceivedDelegate_Handle(FOnPartyJIPRequestReceivedDelegate::CreateUObject(this, &USocialParty::HandlePartyJIPRequestReceived));
 	PartyInterface->AddOnQueryPartyJoinabilityReceivedDelegate_Handle(FOnQueryPartyJoinabilityReceivedDelegate::CreateUObject(this, &USocialParty::HandleJoinabilityQueryReceived));
@@ -429,19 +429,19 @@ void USocialParty::InitializePartyInternal()
 
 	PartyInterface->AddOnPartyMemberJoinedDelegate_Handle(FOnPartyMemberJoinedDelegate::CreateUObject(this, &USocialParty::HandlePartyMemberJoined));
 	PartyInterface->AddOnPartyJIPDelegate_Handle(FOnPartyJIPDelegate::CreateUObject(this, &USocialParty::HandlePartyMemberJIP));
-	PartyInterface->AddOnPartyMemberDataReceivedDelegate_Handle(FOnPartyMemberDataReceivedDelegate::CreateUObject(this, &USocialParty::HandlePartyMemberDataReceived));
+	PartyInterface->AddOnPartyMemberDataReceivedDelegate_Handle(FOnPartyMemberDataReceivedConstDelegate::CreateUObject(this, &USocialParty::HandlePartyMemberDataReceived));
 	PartyInterface->AddOnPartyMemberPromotedDelegate_Handle(FOnPartyMemberPromotedDelegate::CreateUObject(this, &USocialParty::HandlePartyMemberPromoted));
 	PartyInterface->AddOnPartyMemberExitedDelegate_Handle(FOnPartyMemberExitedDelegate::CreateUObject(this, &USocialParty::HandlePartyMemberExited));
 	PartyInterface->AddOnPartySystemStateChangeDelegate_Handle(FOnPartySystemStateChangeDelegate::CreateUObject(this, &USocialParty::HandlePartySystemStateChange));
 	// Create a UPartyMember for every existing member on the OSS party
-	TArray<TSharedRef<FOnlinePartyMember>> OssPartyMembers;
+	TArray<FOnlinePartyMemberConstRef> OssPartyMembers;
 	PartyInterface->GetPartyMembers(*OwningLocalUserId, GetPartyId(), OssPartyMembers);
 	// Always initialize the local member first
-	if (ensure(OssPartyMembers.RemoveAll([this](const TSharedRef<FOnlinePartyMember>& Member) { return *Member->GetUserId() == *OwningLocalUserId; } ) > 0))
+	if (ensure(OssPartyMembers.RemoveAll([this](const FOnlinePartyMemberConstRef& Member) { return *Member->GetUserId() == *OwningLocalUserId; } ) > 0))
 	{
 		GetOrCreatePartyMember(*OwningLocalUserId);
 	}
-	for (TSharedRef<FOnlinePartyMember>& OssMember : OssPartyMembers)
+	for (FOnlinePartyMemberConstRef& OssMember : OssPartyMembers)
 	{
 		GetOrCreatePartyMember(*OssMember->GetUserId());
 	}
@@ -567,7 +567,7 @@ UPartyMember* USocialParty::GetOrCreatePartyMember(const FUniqueNetId& MemberId)
 			{
 				const FOnlinePartyId& PartyId = GetPartyId();
 				const IOnlinePartyPtr PartyInterface = Online::GetPartyInterfaceChecked(GetWorld());
-				const TSharedPtr<FOnlinePartyMember> OssPartyMember = PartyInterface->GetPartyMember(*OwningLocalUserId, PartyId, MemberId);
+				const FOnlinePartyMemberConstPtr OssPartyMember = PartyInterface->GetPartyMember(*OwningLocalUserId, PartyId, MemberId);
 				if (OssPartyMember.IsValid())
 				{
 					PartyMember = NewObject<UPartyMember>(this, PartyMemberClass);
@@ -696,12 +696,12 @@ void USocialParty::HandleJoinabilityQueryReceived(const FUniqueNetId& LocalUserI
 	}
 }
 
-void USocialParty::HandlePartyDataReceived(const FUniqueNetId& LocalUserId, const FOnlinePartyId& PartyId, const TSharedRef<FOnlinePartyData>& PartyData)
+void USocialParty::HandlePartyDataReceived(const FUniqueNetId& LocalUserId, const FOnlinePartyId& PartyId, const FOnlinePartyData& PartyData)
 {
 	if (PartyId == GetPartyId())
 	{
 		check(PartyDataReplicator.IsValid());
-		PartyDataReplicator.ProcessReceivedData(*PartyData);
+		PartyDataReplicator.ProcessReceivedData(PartyData);
 		if (!bHasReceivedRepData)
 		{
 			bHasReceivedRepData = true;
@@ -710,7 +710,7 @@ void USocialParty::HandlePartyDataReceived(const FUniqueNetId& LocalUserId, cons
 	}
 }
 
-void USocialParty::HandlePartyMemberDataReceived(const FUniqueNetId& LocalUserId, const FOnlinePartyId& PartyId, const FUniqueNetId& MemberId, const TSharedRef<FOnlinePartyData>& PartyMemberData)
+void USocialParty::HandlePartyMemberDataReceived(const FUniqueNetId& LocalUserId, const FOnlinePartyId& PartyId, const FUniqueNetId& MemberId, const FOnlinePartyData& PartyMemberData)
 {
 	if (PartyId == GetPartyId())
 	{
@@ -722,7 +722,7 @@ void USocialParty::HandlePartyMemberDataReceived(const FUniqueNetId& LocalUserId
 	}
 }
 
-void USocialParty::HandlePartyConfigChanged(const FUniqueNetId& LocalUserId, const FOnlinePartyId& PartyId, const TSharedRef<FPartyConfiguration>& PartyConfig)
+void USocialParty::HandlePartyConfigChanged(const FUniqueNetId& LocalUserId, const FOnlinePartyId& PartyId, const FPartyConfiguration& PartyConfig)
 {
 	if (PartyId == GetPartyId())
 	{
diff --git a/Engine/Plugins/Online/OnlineFramework/Source/Party/Private/SocialDebugTools.cpp b/Engine/Plugins/Online/OnlineFramework/Source/Party/Private/SocialDebugTools.cpp
--- a/Engine/Plugins/Online/OnlineFramework/Source/Party/Private/SocialDebugTools.cpp
+++ b/Engine/Plugins/Online/OnlineFramework/Source/Party/Private/SocialDebugTools.cpp
@@ -229,7 +229,7 @@ void USocialDebugTools::JoinParty(const FString& Instance, const FString& Friend
 							{
 								if (bWasSuccessful)
 								{
-									TSharedPtr<IOnlinePartyJoinInfo> JoinInfo = OnlineParty->GetAdvertisedParty(UserId, FoundUserId, IOnlinePartySystem::GetPrimaryPartyTypeId());
+									IOnlinePartyJoinInfoConstPtr JoinInfo = OnlineParty->GetAdvertisedParty(UserId, FoundUserId, IOnlinePartySystem::GetPrimaryPartyTypeId());
 									if (JoinInfo.IsValid())
 									{
 										OnlineParty->JoinParty(UserId, *JoinInfo, FOnJoinPartyComplete::CreateLambda([this, Instance, OnlineParty, OnComplete](const FUniqueNetId& UserIdTmp, const FOnlinePartyId& PartyId, const EJoinPartyCompletionResult Result, const int32 NotApprovedReason)
@@ -237,7 +237,7 @@ void USocialDebugTools::JoinParty(const FString& Instance, const FString& Friend
 											bool bSuccess = Result == EJoinPartyCompletionResult::Succeeded;
 											if (bSuccess)
 											{
-												TSharedPtr<FOnlinePartyData> PartyMemberData = GetContext(Instance).GetPartyMemberData();
+												FOnlinePartyDataConstPtr PartyMemberData = GetContext(Instance).GetPartyMemberData();
 												if (PartyMemberData.IsValid())
 												{
 													OnlineParty->UpdatePartyMemberData(UserIdTmp, PartyId, *PartyMemberData);
@@ -259,7 +259,7 @@ void USocialDebugTools::JoinParty(const FString& Instance, const FString& Friend
 					}					
 					else
 					{	
-						TSharedPtr<IOnlinePartyJoinInfo> JoinInfo = GetDefaultPartyJoinInfo();
+						IOnlinePartyJoinInfoConstPtr JoinInfo = GetDefaultPartyJoinInfo();
 						if (JoinInfo.IsValid())
 						{
 							OnlineParty->JoinParty(*LocalUserId, *JoinInfo, FOnJoinPartyComplete::CreateLambda([this, Instance, OnlineParty, OnComplete](const FUniqueNetId& UserId, const FOnlinePartyId& PartyId, const EJoinPartyCompletionResult Result, const int32 NotApprovedReason)
@@ -267,7 +267,7 @@ void USocialDebugTools::JoinParty(const FString& Instance, const FString& Friend
 								bool bSuccess = Result == EJoinPartyCompletionResult::Succeeded;
 								if (bSuccess)
 								{
-									TSharedPtr<FOnlinePartyData> PartyMemberData = GetContext(Instance).GetPartyMemberData();
+									FOnlinePartyDataConstPtr PartyMemberData = GetContext(Instance).GetPartyMemberData();
 									if (PartyMemberData.IsValid())
 									{
 										OnlineParty->UpdatePartyMemberData(UserId, PartyId, *PartyMemberData);
@@ -460,9 +460,9 @@ USocialManager& USocialDebugTools::GetSocialManager() const
 	return *OuterSocialManager;
 }
 
-TSharedPtr<IOnlinePartyJoinInfo> USocialDebugTools::GetDefaultPartyJoinInfo() const
+IOnlinePartyJoinInfoConstPtr USocialDebugTools::GetDefaultPartyJoinInfo() const
 {
-	TSharedPtr<IOnlinePartyJoinInfo> Result;
+	IOnlinePartyJoinInfoConstPtr Result;
 	IOnlineSubsystem* OnlineSub = GetDefaultOSS();
 	if (OnlineSub)
 	{
@@ -718,9 +718,9 @@ void USocialDebugTools::HandlePartyInviteReceived(const FUniqueNetId& LocalUserI
 			IOnlinePartyPtr OnlineParty = Context->GetOSS()->GetPartyInterface();
 			if (OnlineParty.IsValid())
 			{
-				TArray<TSharedRef<IOnlinePartyJoinInfo>> PartyInvites;
+				TArray<TSharedRef<const IOnlinePartyJoinInfo>> PartyInvites;
 				OnlineParty->GetPendingInvites(LocalUserId, PartyInvites);
-				for (const TSharedRef<IOnlinePartyJoinInfo>& Invite : PartyInvites)
+				for (const TSharedRef<const IOnlinePartyJoinInfo>& Invite : PartyInvites)
 				{
 					if (*Invite->GetPartyId() == PartyId)
 					{						
@@ -730,7 +730,7 @@ void USocialDebugTools::HandlePartyInviteReceived(const FUniqueNetId& LocalUserI
 							bool bSuccess = Result == EJoinPartyCompletionResult::Succeeded;
 							if (bSuccess)
 							{
-								TSharedPtr<FOnlinePartyData> PartyMemberData = GetContext(Instance).GetPartyMemberData();
+								FOnlinePartyDataConstPtr PartyMemberData = GetContext(Instance).GetPartyMemberData();
 								if (PartyMemberData.IsValid())
 								{
 									OnlineParty->UpdatePartyMemberData(UserId, PartyIdTmp, *PartyMemberData);
diff --git a/Engine/Plugins/Online/OnlineFramework/Source/Party/Private/SocialToolkit.cpp b/Engine/Plugins/Online/OnlineFramework/Source/Party/Private/SocialToolkit.cpp
--- a/Engine/Plugins/Online/OnlineFramework/Source/Party/Private/SocialToolkit.cpp
+++ b/Engine/Plugins/Online/OnlineFramework/Source/Party/Private/SocialToolkit.cpp
@@ -976,9 +976,9 @@ void USocialToolkit::HandleExistingPartyInvites(ESocialSubsystem SubsystemType)
 			FUniqueNetIdRepl LocalUserId = GetLocalUserNetId(SubsystemType);
 			if (PartyInterface.IsValid() && LocalUserId.IsValid())
 			{
-				TArray<TSharedRef<IOnlinePartyJoinInfo>> PendingInvites;
+				TArray<IOnlinePartyJoinInfoConstRef> PendingInvites;
 				PartyInterface->GetPendingInvites(*LocalUserId, PendingInvites);
-				for (const TSharedRef<IOnlinePartyJoinInfo>& PendingInvite : PendingInvites)
+				for (const IOnlinePartyJoinInfoConstRef& PendingInvite : PendingInvites)
 				{
 					HandlePartyInviteReceived(*LocalUserId, *PendingInvite->GetPartyId(), *PendingInvite->GetSourceUserId());
 				}
diff --git a/Engine/Plugins/Online/OnlineFramework/Source/Party/Private/User/SocialUser.cpp b/Engine/Plugins/Online/OnlineFramework/Source/Party/Private/User/SocialUser.cpp
--- a/Engine/Plugins/Online/OnlineFramework/Source/Party/Private/User/SocialUser.cpp
+++ b/Engine/Plugins/Online/OnlineFramework/Source/Party/Private/User/SocialUser.cpp
@@ -885,10 +885,10 @@ TSharedPtr<const IOnlinePartyJoinInfo> USocialUser::GetPartyJoinInfo(const FOnli
 			if (!JoinInfo.IsValid())
 			{
 				// No advertised party info, check to see if this user has sent an invite
-				TArray<TSharedRef<IOnlinePartyJoinInfo>> AllPendingInvites;
+				TArray<IOnlinePartyJoinInfoConstRef> AllPendingInvites;
 				if (PartyInterface->GetPendingInvites(*LocalUserId, AllPendingInvites))
 				{
-					for (const TSharedRef<IOnlinePartyJoinInfo>& InvitationJoinInfo : AllPendingInvites)
+					for (const IOnlinePartyJoinInfoConstRef& InvitationJoinInfo : AllPendingInvites)
 					{
 						if (*InvitationJoinInfo->GetSourceUserId() == *UserId)
 						{
@@ -914,10 +914,10 @@ bool USocialUser::HasSentPartyInvite(const FOnlinePartyTypeId& PartyTypeId) cons
 		const FUniqueNetIdRepl UserId = GetUserId(ESocialSubsystem::Primary);
 		if (ensure(LocalUserId.IsValid()) && UserId.IsValid())
 		{
-			TArray<TSharedRef<IOnlinePartyJoinInfo>> AllPendingInvites;
+			TArray<IOnlinePartyJoinInfoConstRef> AllPendingInvites;
 			if (PartyInterface->GetPendingInvites(*LocalUserId, AllPendingInvites))
 			{
-				for (const TSharedRef<IOnlinePartyJoinInfo>& InvitationJoinInfo : AllPendingInvites)
+				for (const IOnlinePartyJoinInfoConstRef& InvitationJoinInfo : AllPendingInvites)
 				{
 					if (*InvitationJoinInfo->GetSourceUserId() == *UserId && InvitationJoinInfo->GetPartyTypeId() == PartyTypeId)
 					{
diff --git a/Engine/Plugins/Online/OnlineFramework/Source/Party/Private/User/SocialUserList.h b/Engine/Plugins/Online/OnlineFramework/Source/Party/Private/User/SocialUserList.h
--- a/Engine/Plugins/Online/OnlineFramework/Source/Party/Private/User/SocialUserList.h
+++ b/Engine/Plugins/Online/OnlineFramework/Source/Party/Private/User/SocialUserList.h
@@ -5,7 +5,7 @@
 #include "User/ISocialUserList.h"
 #include "UObject/GCObject.h"
 
-enum class EMemberExitedReason;
+enum class EMemberExitedReason : uint8;
 
 class FSocialUserList : public ISocialUserList, public FGCObject, public TSharedFromThis<FSocialUserList>
 {
diff --git a/Engine/Plugins/Online/OnlineFramework/Source/Party/Public/Chat/SocialPartyChatRoom.h b/Engine/Plugins/Online/OnlineFramework/Source/Party/Public/Chat/SocialPartyChatRoom.h
--- a/Engine/Plugins/Online/OnlineFramework/Source/Party/Public/Chat/SocialPartyChatRoom.h
+++ b/Engine/Plugins/Online/OnlineFramework/Source/Party/Public/Chat/SocialPartyChatRoom.h
@@ -6,7 +6,7 @@
 #include "SocialPartyChatRoom.generated.h"
 
 class USocialChatManager;
-enum class EMemberExitedReason;
+enum class EMemberExitedReason : uint8;
 
 /** A multi-user chat room channel. Used for all chat situations outside of private user-to-user direct messages. */
 UCLASS()
diff --git a/Engine/Plugins/Online/OnlineFramework/Source/Party/Public/Party/PartyMember.h b/Engine/Plugins/Online/OnlineFramework/Source/Party/Public/Party/PartyMember.h
--- a/Engine/Plugins/Online/OnlineFramework/Source/Party/Public/Party/PartyMember.h
+++ b/Engine/Plugins/Online/OnlineFramework/Source/Party/Public/Party/PartyMember.h
@@ -10,8 +10,8 @@
 class USocialUser;
 class FOnlinePartyMember;
 class FOnlinePartyData;
-enum class EMemberExitedReason;
-enum class EMemberConnectionStatus;
+enum class EMemberExitedReason : uint8;
+enum class EMemberConnectionStatus : uint8;
 
 /** Base struct used to replicate data about the state of a single party member to all members. */
 USTRUCT()
@@ -99,10 +99,10 @@ public:
 	FString ToDebugString(bool bIncludePartyId = true) const;
 
 PARTY_SCOPE:
-	void InitializePartyMember(const TSharedRef<FOnlinePartyMember>& OssMember, const FSimpleDelegate& OnInitComplete);
+	void InitializePartyMember(const FOnlinePartyMemberConstRef& OssMember, const FSimpleDelegate& OnInitComplete);
 
 	FPartyMemberRepData& GetMutableRepData() { return *MemberDataReplicator; }
-	void NotifyMemberDataReceived(const TSharedRef<FOnlinePartyData>& MemberData);
+	void NotifyMemberDataReceived(const FOnlinePartyData& MemberData);
 	void NotifyMemberPromoted();
 	void NotifyMemberDemoted();
 	void NotifyRemovedFromParty(EMemberExitedReason ExitReason);
@@ -125,7 +125,7 @@ private:
 	void HandleSocialUserInitialized(USocialUser& InitializedUser);
 	void HandleMemberConnectionStatusChanged(const FUniqueNetId& ChangedUserId, const EMemberConnectionStatus NewMemberConnectionStatus, const EMemberConnectionStatus PreviousMemberConnectionStatus);
 
-	TSharedPtr<FOnlinePartyMember> OssPartyMember;
+	FOnlinePartyMemberConstPtr OssPartyMember;
 
 	UPROPERTY()
 	USocialUser* SocialUser = nullptr;
diff --git a/Engine/Plugins/Online/OnlineFramework/Source/Party/Public/Party/PartyTypes.h b/Engine/Plugins/Online/OnlineFramework/Source/Party/Public/Party/PartyTypes.h
--- a/Engine/Plugins/Online/OnlineFramework/Source/Party/Public/Party/PartyTypes.h
+++ b/Engine/Plugins/Online/OnlineFramework/Source/Party/Public/Party/PartyTypes.h
@@ -12,7 +12,7 @@ class UPartyMember;
 class IOnlinePartySystem;
 class FOnlinePartyData;
 
-enum class EJoinPartyCompletionResult;
+enum class EJoinPartyCompletionResult : int8;
 
 UENUM()
 enum class EPartyType : uint8
diff --git a/Engine/Plugins/Online/OnlineFramework/Source/Party/Public/Party/SocialParty.h b/Engine/Plugins/Online/OnlineFramework/Source/Party/Public/Party/SocialParty.h
--- a/Engine/Plugins/Online/OnlineFramework/Source/Party/Public/Party/SocialParty.h
+++ b/Engine/Plugins/Online/OnlineFramework/Source/Party/Public/Party/SocialParty.h
@@ -23,7 +23,7 @@ class USocialUser;
 
 class FOnlineSessionSettings;
 class FOnlineSessionSearchResult;
-enum class EMemberExitedReason;
+enum class EMemberExitedReason : uint8;
 
 /** Base struct used to replicate data about the state of the party to all members. */
 USTRUCT()
@@ -302,15 +302,15 @@ private:
 
 private:	// Handlers
 	void HandlePartyStateChanged(const FUniqueNetId& LocalUserId, const FOnlinePartyId& PartyId, EPartyState PartyState, EPartyState PreviousPartyState);
-	void HandlePartyConfigChanged(const FUniqueNetId& LocalUserId, const FOnlinePartyId& PartyId, const TSharedRef<FPartyConfiguration>& PartyConfig);
+	void HandlePartyConfigChanged(const FUniqueNetId& LocalUserId, const FOnlinePartyId& PartyId, const FPartyConfiguration& PartyConfig);
 	void HandleUpdatePartyConfigComplete(const FUniqueNetId& LocalUserId, const FOnlinePartyId& PartyId, EUpdateConfigCompletionResult Result);
-	void HandlePartyDataReceived(const FUniqueNetId& LocalUserId, const FOnlinePartyId& PartyId, const TSharedRef<FOnlinePartyData>& PartyData);
+	void HandlePartyDataReceived(const FUniqueNetId& LocalUserId, const FOnlinePartyId& PartyId, const FOnlinePartyData& PartyData);
 	void HandleJoinabilityQueryReceived(const FUniqueNetId& LocalUserId, const FOnlinePartyId& PartyId, const FUniqueNetId& SenderId, const FString& Platform, const FOnlinePartyData& JoinData);
 	void HandlePartyJoinRequestReceived(const FUniqueNetId& LocalUserId, const FOnlinePartyId& PartyId, const FUniqueNetId& SenderId, const FString& Platform, const FOnlinePartyData& JoinData);
 	void HandlePartyJIPRequestReceived(const FUniqueNetId& LocalUserId, const FOnlinePartyId& PartyId, const FUniqueNetId& SenderId);
 	void HandlePartyLeft(const FUniqueNetId& LocalUserId, const FOnlinePartyId& PartyId);
 	void HandlePartyMemberExited(const FUniqueNetId& LocalUserId, const FOnlinePartyId& PartyId, const FUniqueNetId& MemberId, EMemberExitedReason ExitReason);
-	void HandlePartyMemberDataReceived(const FUniqueNetId& LocalUserId, const FOnlinePartyId& PartyId, const FUniqueNetId& MemberId, const TSharedRef<FOnlinePartyData>& PartyMemberData);
+	void HandlePartyMemberDataReceived(const FUniqueNetId& LocalUserId, const FOnlinePartyId& PartyId, const FUniqueNetId& MemberId, const FOnlinePartyData& PartyMemberData);
 	void HandlePartyMemberJoined(const FUniqueNetId& LocalUserId, const FOnlinePartyId& PartyId, const FUniqueNetId& MemberId);
 	void HandlePartyMemberJIP(const FUniqueNetId& LocalUserId, const FOnlinePartyId& PartyId, bool Success);
 	void HandlePartyMemberPromoted(const FUniqueNetId& LocalUserId, const FOnlinePartyId& PartyId, const FUniqueNetId& NewLeaderId);
diff --git a/Engine/Plugins/Online/OnlineFramework/Source/Party/Public/SocialDebugTools.h b/Engine/Plugins/Online/OnlineFramework/Source/Party/Public/SocialDebugTools.h
--- a/Engine/Plugins/Online/OnlineFramework/Source/Party/Public/SocialDebugTools.h
+++ b/Engine/Plugins/Online/OnlineFramework/Source/Party/Public/SocialDebugTools.h
@@ -7,9 +7,10 @@
 #include "SocialDebugTools.generated.h"
 
 class IOnlineSubsystem;
-class IOnlinePartyJoinInfo;
 class FOnlineAccountCredentials;
-class FOnlinePartyData;
+typedef TSharedPtr<const class IOnlinePartyJoinInfo> IOnlinePartyJoinInfoConstPtr;
+typedef TSharedPtr<class FOnlinePartyData> FOnlinePartyDataPtr;
+typedef TSharedPtr<const class FOnlinePartyData> FOnlinePartyDataConstPtr;
 
 UCLASS(Within = SocialManager, Config = Game)
 class PARTY_API USocialDebugTools : public UObject, public FExec
@@ -61,12 +62,12 @@ public:
 		void Init();
 		void Shutdown();
 		inline IOnlineSubsystem* GetOSS() { return OnlineSub; }
-		inline TSharedPtr<FOnlinePartyData> GetPartyMemberData() { return PartyMemberData; }
+		inline FOnlinePartyDataConstPtr GetPartyMemberData() { return PartyMemberData; }
 
 		FString Name;
 		IOnlineSubsystem* OnlineSub;
 		USocialDebugTools& Owner;
-		TSharedPtr<FOnlinePartyData> PartyMemberData;
+		FOnlinePartyDataPtr PartyMemberData;
 
 		// delegates
 		FDelegateHandle LoginCompleteDelegateHandle;
@@ -91,7 +92,7 @@ private:
 
 	TMap<FString, FInstanceContext> Contexts;
 
-	TSharedPtr<IOnlinePartyJoinInfo> GetDefaultPartyJoinInfo() const;
+	IOnlinePartyJoinInfoConstPtr GetDefaultPartyJoinInfo() const;
 	IOnlineSubsystem* GetDefaultOSS() const;
 	void PrintExecUsage() const;
 	virtual void PrintExecCommands() const;
diff --git a/Engine/Plugins/Online/OnlineFramework/Source/Party/Public/SocialToolkit.h b/Engine/Plugins/Online/OnlineFramework/Source/Party/Public/SocialToolkit.h
--- a/Engine/Plugins/Online/OnlineFramework/Source/Party/Public/SocialToolkit.h
+++ b/Engine/Plugins/Online/OnlineFramework/Source/Party/Public/SocialToolkit.h
@@ -17,7 +17,7 @@ class USocialChatManager;
 
 class FOnlinePartyId;
 
-enum class EMemberExitedReason;
+enum class EMemberExitedReason : uint8;
 
 namespace EOnlinePresenceState { enum Type : uint8; }
 
diff --git a/Engine/Plugins/Online/OnlineSubsystem/Source/Private/OnlinePartyInterface.cpp b/Engine/Plugins/Online/OnlineSubsystem/Source/Private/OnlinePartyInterface.cpp
--- a/Engine/Plugins/Online/OnlineSubsystem/Source/Private/OnlinePartyInterface.cpp
+++ b/Engine/Plugins/Online/OnlineSubsystem/Source/Private/OnlinePartyInterface.cpp
@@ -138,6 +138,81 @@ FString FPartyInvitationRecipient::ToDebugString() const
 	return FString::Printf(TEXT("Id=[%s], PlatformData=[%s]"), *Id->ToDebugString(), *PlatformData);
 }
 
+bool IOnlinePartySystem::GetPartyMembers(const FUniqueNetId& LocalUserId, const FOnlinePartyId& PartyId, TArray<TSharedRef<FOnlinePartyMember>>& OutPartyMembersArray) const
+{
+	TArray<FOnlinePartyMemberConstRef> ConstMembers;
+	const bool bResult = GetPartyMembers(LocalUserId, PartyId, ConstMembers);
+	OutPartyMembersArray.Empty(ConstMembers.Num());
+	if (bResult)
+	{
+		for (const FOnlinePartyMemberConstRef& ConstMember : ConstMembers)
+		{
+			OutPartyMembersArray.Emplace(ConstCastSharedRef<FOnlinePartyMember>(ConstMember));
+		}
+	}
+	return bResult;
+}
+
+bool IOnlinePartySystem::GetPendingInvites(const FUniqueNetId& LocalUserId, TArray<TSharedRef<IOnlinePartyJoinInfo>>& OutPendingInvitesArray) const
+{
+	TArray<IOnlinePartyJoinInfoConstRef> ConstJoinInfos;
+	const bool bResult = GetPendingInvites(LocalUserId, ConstJoinInfos);
+	OutPendingInvitesArray.Empty(ConstJoinInfos.Num());
+	if (bResult)
+	{
+		for (const IOnlinePartyJoinInfoConstRef& ConstJoinInfo : ConstJoinInfos)
+		{
+			OutPendingInvitesArray.Emplace(ConstCastSharedRef<IOnlinePartyJoinInfo>(ConstJoinInfo));
+		}
+	}
+	return bResult;
+}
+
+bool IOnlinePartySystem::GetPendingJoinRequests(const FUniqueNetId& LocalUserId, const FOnlinePartyId& PartyId, TArray<TSharedRef<IOnlinePartyPendingJoinRequestInfo>>& OutPendingJoinRequestArray) const
+{
+	TArray<IOnlinePartyPendingJoinRequestInfoConstRef> ConstJoinRequests;
+	const bool bResult = GetPendingJoinRequests(LocalUserId, PartyId, ConstJoinRequests);
+	OutPendingJoinRequestArray.Empty(ConstJoinRequests.Num());
+	if (bResult)
+	{
+		for (const IOnlinePartyPendingJoinRequestInfoConstRef& ConstJoinRequest : ConstJoinRequests)
+		{
+			OutPendingJoinRequestArray.Emplace(ConstCastSharedRef<IOnlinePartyPendingJoinRequestInfo>(ConstJoinRequest));
+		}
+	}
+	return bResult;
+}
+
+FDelegateHandle IOnlinePartySystem::AddOnPartyConfigChangedDelegate_Handle(const FOnPartyConfigChangedDelegate& Delegate)
+{
+	auto DeprecationHelperLambda = [Delegate](const FUniqueNetId& LocalUserId, const FOnlinePartyId& PartyId, const FPartyConfiguration& PartyConfig)
+	{
+		TSharedRef<FPartyConfiguration> PartyConfigurationRef((const_cast<FPartyConfiguration&>(PartyConfig)).AsShared());
+		Delegate.ExecuteIfBound(LocalUserId, PartyId, PartyConfigurationRef);
+	};
+	return OnPartyConfigChangedDelegates.Add(FOnPartyConfigChangedConstDelegate::CreateLambda(DeprecationHelperLambda));
+}
+
+FDelegateHandle IOnlinePartySystem::AddOnPartyDataReceivedDelegate_Handle(const FOnPartyDataReceivedDelegate& Delegate)
+{
+	auto DeprecationHelperLambda = [Delegate](const FUniqueNetId& LocalUserId, const FOnlinePartyId& PartyId, const FOnlinePartyData& PartyData)
+	{
+		TSharedRef<FOnlinePartyData> PartyDataRef((const_cast<FOnlinePartyData&>(PartyData)).AsShared());
+		Delegate.ExecuteIfBound(LocalUserId, PartyId, PartyDataRef);
+	};
+	return OnPartyDataReceivedDelegates.Add(FOnPartyDataReceivedConstDelegate::CreateLambda(DeprecationHelperLambda));
+}
+
+FDelegateHandle IOnlinePartySystem::AddOnPartyMemberDataReceivedDelegate_Handle(const FOnPartyMemberDataReceivedDelegate& Delegate)
+{
+	auto DeprecationHelperLambda = [Delegate](const FUniqueNetId& LocalUserId, const FOnlinePartyId& PartyId, const FUniqueNetId& MemberId, const FOnlinePartyData& PartyData)
+	{
+		TSharedRef<FOnlinePartyData> PartyDataRef((const_cast<FOnlinePartyData&>(PartyData)).AsShared());
+		Delegate.ExecuteIfBound(LocalUserId, PartyId, MemberId, PartyDataRef);
+	};
+	return OnPartyMemberDataReceivedDelegates.Add(FOnPartyMemberDataReceivedConstDelegate::CreateLambda(DeprecationHelperLambda));
+}
+
 bool FPartyConfiguration::operator==(const FPartyConfiguration& Other) const
 {
 	return JoinRequestAction == Other.JoinRequestAction &&
diff --git a/Engine/Plugins/Online/OnlineSubsystem/Source/Public/Interfaces/OnlinePartyInterface.h b/Engine/Plugins/Online/OnlineSubsystem/Source/Public/Interfaces/OnlinePartyInterface.h
--- a/Engine/Plugins/Online/OnlineSubsystem/Source/Public/Interfaces/OnlinePartyInterface.h
+++ b/Engine/Plugins/Online/OnlineSubsystem/Source/Public/Interfaces/OnlinePartyInterface.h
@@ -9,6 +9,7 @@
 #include "OnlineDelegateMacros.h"
 
 typedef FString FChatRoomId;
+struct FOnlineError;
 
 ONLINESUBSYSTEM_API DECLARE_LOG_CATEGORY_EXTERN(LogOnlineParty, Log, All);
 #define UE_LOG_ONLINE_PARTY(Verbosity, Format, ...) \
@@ -24,21 +25,22 @@ ONLINESUBSYSTEM_API DECLARE_LOG_CATEGORY_EXTERN(LogOnlineParty, Log, All);
 #define F_PREFIX(TypeToPrefix) F##TypeToPrefix
 #define PARTY_DECLARE_DELEGATETYPE(Type) typedef F##Type::FDelegate F##Type##Delegate
 
-enum class EAcceptPartyInvitationCompletionResult;
-enum class ECreatePartyCompletionResult;
-enum class EInvitationResponse;
-enum class EJoinPartyCompletionResult;
-enum class EKickMemberCompletionResult;
-enum class ELeavePartyCompletionResult;
-enum class EPromoteMemberCompletionResult;
-enum class ERejectPartyInvitationCompletionResult;
-enum class ERequestPartyInvitationCompletionResult;
-enum class ESendPartyInvitationCompletionResult;
-enum class EUpdateConfigCompletionResult;
+enum class EAcceptPartyInvitationCompletionResult : int8;
+enum class ECreatePartyCompletionResult : int8;
+enum class EJoinPartyCompletionResult : int8;
+enum class EKickMemberCompletionResult : int8;
+enum class ELeavePartyCompletionResult : int8;
+enum class EPromoteMemberCompletionResult : int8;
+enum class ERejectPartyInvitationCompletionResult : int8;
+enum class ERequestPartyInvitationCompletionResult : int8;
+enum class ESendPartyInvitationCompletionResult : int8;
+enum class EUpdateConfigCompletionResult : int8;
+
+enum class EInvitationResponse : uint8;
 
 struct FAnalyticsEventAttribute;
 
-enum class EMemberConnectionStatus
+enum class EMemberConnectionStatus : uint8
 {
 	Uninitialized,
 	Disconnected,
@@ -46,30 +48,24 @@ enum class EMemberConnectionStatus
 	Connected
 };
 
-/**
-* notification when a party member's data status has changed
-* @param ChangedUserId - id associated with this notification
-* @param NewMemberConnectionStatus - new member data status
-*/
-DECLARE_MULTICAST_DELEGATE_ThreeParams(F_PREFIX(OnMemberConnectionStatusChanged), const FUniqueNetId& /*ChangedUserId*/, const EMemberConnectionStatus /*NewMemberDataStatus*/, const EMemberConnectionStatus /*PreviousMemberDataStatus*/);
-PARTY_DECLARE_DELEGATETYPE(OnMemberConnectionStatusChanged);
-
 /**
  * Party member user info returned by IOnlineParty interface
  */
-class FOnlinePartyMember : public FOnlineUser
+class FOnlinePartyMember
+	: public FOnlineUser
 {
 public:
 	EMemberConnectionStatus MemberConnectionStatus = EMemberConnectionStatus::Uninitialized;
 	EMemberConnectionStatus PreviousMemberConnectionStatus = EMemberConnectionStatus::Uninitialized;
 
 	/**
-	* notification when a party member's data status has changed
-	* @param ChangedUserId - id associated with this notification
-	* @param NewMemberConnectionStatus - new member data status
-	* @param PreviousMemberConnectionStatus - previous member data status
-	*/
-	DEFINE_ONLINE_DELEGATE_THREE_PARAM(OnMemberConnectionStatusChanged, const FUniqueNetId& /*ChangedUserId*/, const EMemberConnectionStatus /*NewMemberConnectionStatus*/, const EMemberConnectionStatus /*PreviousMemberConnectionStatus*/);
+	 * Event when a party member's connection status has changed
+	 * @param ChangedUserId - id associated with this notification
+	 * @param NewMemberConnectionStatus - new member data status
+	 * @param PreviousMemberConnectionStatus - previous member data status
+	 */
+	DECLARE_EVENT_ThreeParams(FOnlinePartyMember, FOnMemberConnectionStatusChanged, const FUniqueNetId& /*ChangedUserId*/, const EMemberConnectionStatus /*NewMemberConnectionStatus*/, const EMemberConnectionStatus /*PreviousMemberConnectionStatus*/);
+	FOnMemberConnectionStatusChanged& OnMemberConnectionStatusChanged() const { return OnMemberConnectionStatusChangedEvent; }
 
 	void SetMemberConnectionStatus(EMemberConnectionStatus NewMemberConnectionStatus)
 	{
@@ -77,24 +73,27 @@ public:
 		{
 			PreviousMemberConnectionStatus = MemberConnectionStatus;
 			MemberConnectionStatus = NewMemberConnectionStatus;
-			TriggerOnMemberConnectionStatusChangedDelegates(*GetUserId(), MemberConnectionStatus, PreviousMemberConnectionStatus);
+			OnMemberConnectionStatusChangedEvent.Broadcast(*GetUserId(), MemberConnectionStatus, PreviousMemberConnectionStatus);
 		}
 	}
+
+private:
+	/** Event fired when connection status changes */
+	mutable FOnMemberConnectionStatusChanged OnMemberConnectionStatusChangedEvent;
 };
 
+typedef TSharedRef<const FOnlinePartyMember> FOnlinePartyMemberConstRef;
+typedef TSharedPtr<const FOnlinePartyMember> FOnlinePartyMemberConstPtr;
+
 /**
  * Data associated with the entire party
  */
-class ONLINESUBSYSTEM_API FOnlinePartyData : public TSharedFromThis<FOnlinePartyData>
+class ONLINESUBSYSTEM_API FOnlinePartyData
+	: public TSharedFromThis<FOnlinePartyData>
 {
 public:
-	FOnlinePartyData() :
-		TotalBytes(0),
-		TotalEffectiveBytes(0),
-		TotalPackets(0),
-		RevisionCount(0)
-		{}
-	~FOnlinePartyData() {}
+	FOnlinePartyData() = default;
+	virtual ~FOnlinePartyData() = default;
 
 	/**
 	 * Equality operator
@@ -293,14 +292,14 @@ public:
 
 	/** Stat tracking variables */
 	/** Total number of bytes generated by calls to ToJsonFull and ToJsonDirty */
-	mutable int32 TotalBytes;
+	mutable int32 TotalBytes = 0;
 	/** Total number of bytes generated by calls to ToJsonFull and ToJsonDirty, multiplied by the number of recipients the packet was sent to */
-	mutable int32 TotalEffectiveBytes;
+	mutable int32 TotalEffectiveBytes = 0;
 	/** Total number of packets generated by calls to ToJsonFull and ToJsonDirty */
-	mutable int32 TotalPackets;
+	mutable int32 TotalPackets = 0;
 
 	/** Id representing number of updates sent, useful for determining if a client has missed an update */
-	mutable int32 RevisionCount;
+	mutable int32 RevisionCount = 0;
 
 private:
 	/** map of key/val attributes that represents the data */
@@ -311,6 +310,11 @@ private:
 
 };
 
+typedef TSharedRef<FOnlinePartyData> FOnlinePartyDataRef;
+typedef TSharedPtr<FOnlinePartyData> FOnlinePartyDataPtr;
+typedef TSharedRef<const FOnlinePartyData> FOnlinePartyDataConstRef;
+typedef TSharedPtr<const FOnlinePartyData> FOnlinePartyDataConstPtr;
+
 /**
 * Info needed to join a party
 */
@@ -341,6 +345,9 @@ public:
 	virtual TSharedRef<const FOnlinePartyData> GetSenderJoinData() const = 0;
 };
 
+typedef TSharedRef<const IOnlinePartyPendingJoinRequestInfo> IOnlinePartyPendingJoinRequestInfoConstRef;
+typedef TSharedPtr<const IOnlinePartyPendingJoinRequestInfo> IOnlinePartyPendingJoinRequestInfoConstPtr;
+
 /**
  * Info needed to join a party
  */
@@ -429,6 +436,9 @@ public:
 	virtual bool CanRequestAnInvite() const = 0;
 };
 
+typedef TSharedRef<const IOnlinePartyJoinInfo> IOnlinePartyJoinInfoConstRef;
+typedef TSharedPtr<const IOnlinePartyJoinInfo> IOnlinePartyJoinInfoConstPtr;
+
 /**
  * Permissions for party features
  */
@@ -450,7 +460,7 @@ namespace PartySystemPermissions
 	};
 }
 
-enum class EJoinRequestAction
+enum class EJoinRequestAction : uint8
 {
 	Manual,
 	AutoApprove,
@@ -459,7 +469,8 @@ enum class EJoinRequestAction
 /**
  * Options for configuring a new party or for updating an existing party
  */
-struct ONLINESUBSYSTEM_API FPartyConfiguration : public TSharedFromThis<FPartyConfiguration>
+struct ONLINESUBSYSTEM_API FPartyConfiguration
+	: public TSharedFromThis<FPartyConfiguration>
 {
 	FPartyConfiguration()
 		: JoinRequestAction(EJoinRequestAction::Manual)
@@ -511,7 +522,9 @@ struct ONLINESUBSYSTEM_API FPartyConfiguration : public TSharedFromThis<FPartyCo
 	FString Password;
 };
 
-enum class EPartyState
+typedef TSharedRef<const FPartyConfiguration> FPartyConfigurationConstRef;
+
+enum class EPartyState : uint8
 {
 	None,
 	CreatePending,
@@ -526,16 +539,16 @@ enum class EPartyState
 /**
  * Current state associated with a party
  */
-class FOnlineParty : public TSharedFromThis<FOnlineParty>
+class FOnlineParty
+	: public TSharedFromThis<FOnlineParty>
 {
 	FOnlineParty() = delete;
 protected:
-	FOnlineParty(const TSharedRef<const FOnlinePartyId>& InPartyId, const FOnlinePartyTypeId InPartyTypeId, TSharedRef<FPartyConfiguration> InConfig = MakeShared<FPartyConfiguration>())
+		FOnlineParty(const TSharedRef<const FOnlinePartyId>& InPartyId, const FOnlinePartyTypeId InPartyTypeId)
 		: PartyId(InPartyId)
 		, PartyTypeId(InPartyTypeId)
 		, State(EPartyState::None)
 		, PreviousState(EPartyState::None)
-		, Config(InConfig)
 	{}
 
 public:
@@ -578,13 +591,14 @@ public:
 	EPartyState State;
 	/** The current state of the party */
 	EPartyState PreviousState;
-	/** Current state of configuration */
-	TSharedRef<FPartyConfiguration> Config;
 	/** id of chat room associated with the party */
 	FChatRoomId RoomId;
 };
 
-enum class EMemberExitedReason
+typedef TSharedRef<const FOnlineParty> FOnlinePartyConstRef;
+typedef TSharedPtr<const FOnlineParty> FOnlinePartyConstPtr;
+
+enum class EMemberExitedReason : uint8
 {
 	Unknown,
 	Left,
@@ -592,6 +606,18 @@ enum class EMemberExitedReason
 	Kicked
 };
 
+enum class EPartyInvitationRemovedReason : uint8
+{
+	/** Unknown or undefined reason */
+	Unknown,
+	/** Expired */
+	Expired,
+	/** Became invalid (for example, party was destroyed) */
+	Invalidated,
+	/** User declined the invitation */
+	Declined,
+};
+
 /** Recipient information for SendInvitation */
 struct FPartyInvitationRecipient
 {
@@ -614,20 +640,7 @@ struct FPartyInvitationRecipient
 	FString ONLINESUBSYSTEM_API ToDebugString() const;
 };
 
-struct FOnlineError;
-enum class ECreatePartyCompletionResult;
-enum class EJoinPartyCompletionResult;
-enum class ELeavePartyCompletionResult;
-enum class EUpdateConfigCompletionResult;
-enum class ERequestPartyInvitationCompletionResult;
-enum class ESendPartyInvitationCompletionResult;
-enum class EAcceptPartyInvitationCompletionResult;
-enum class ERejectPartyInvitationCompletionResult;
-enum class EKickMemberCompletionResult;
-enum class EPromoteMemberCompletionResult;
-enum class EInvitationResponse;
-
-enum class EPartySystemState
+enum class EPartySystemState : uint8
 {
 	Initializing = 0,
 	Initialized,
@@ -800,7 +813,8 @@ DECLARE_MULTICAST_DELEGATE_ThreeParams(F_PREFIX(OnPartyPromotionLockoutChanged),
 PARTY_DECLARE_DELEGATETYPE(OnPartyPromotionLockoutChanged);
 
 /**
- * Notification when party data is updated
+ * Notification when party config is updated
+ * Deprecated - Use OnPartyConfigChangedConst
  * @param LocalUserId - id associated with this notification
  * @param PartyId - id associated with the party
  * @param PartyConfig - party whose config was updated
@@ -808,8 +822,18 @@ PARTY_DECLARE_DELEGATETYPE(OnPartyPromotionLockoutChanged);
 DECLARE_MULTICAST_DELEGATE_ThreeParams(F_PREFIX(OnPartyConfigChanged), const FUniqueNetId& /*LocalUserId*/, const FOnlinePartyId& /*PartyId*/, const TSharedRef<FPartyConfiguration>& /*PartyConfig*/);
 PARTY_DECLARE_DELEGATETYPE(OnPartyConfigChanged);
 
+/**
+ * Notification when party config is updated
+ * @param LocalUserId - id associated with this notification
+ * @param PartyId - id associated with the party
+ * @param PartyConfig - party whose config was updated
+ */
+DECLARE_MULTICAST_DELEGATE_ThreeParams(F_PREFIX(OnPartyConfigChangedConst), const FUniqueNetId& /*LocalUserId*/, const FOnlinePartyId& /*PartyId*/, const FPartyConfiguration& /*PartyConfig*/);
+PARTY_DECLARE_DELEGATETYPE(OnPartyConfigChangedConst);
+
 /**
  * Notification when party data is updated
+ * Deprecated - Use OnPartyDataReceivedConst
  * @param LocalUserId - id associated with this notification
  * @param PartyId - id associated with the party
  * @param PartyData - party data that was updated
@@ -817,6 +841,15 @@ PARTY_DECLARE_DELEGATETYPE(OnPartyConfigChanged);
 DECLARE_MULTICAST_DELEGATE_ThreeParams(F_PREFIX(OnPartyDataReceived), const FUniqueNetId& /*LocalUserId*/, const FOnlinePartyId& /*PartyId*/, const TSharedRef<FOnlinePartyData>& /*PartyData*/);
 PARTY_DECLARE_DELEGATETYPE(OnPartyDataReceived);
 
+/**
+ * Notification when party data is updated
+ * @param LocalUserId - id associated with this notification
+ * @param PartyId - id associated with the party
+ * @param PartyData - party data that was updated
+ */
+DECLARE_MULTICAST_DELEGATE_ThreeParams(F_PREFIX(OnPartyDataReceivedConst), const FUniqueNetId& /*LocalUserId*/, const FOnlinePartyId& /*PartyId*/, const FOnlinePartyData& /*PartyData*/);
+PARTY_DECLARE_DELEGATETYPE(OnPartyDataReceivedConst);
+
 /**
  * Notification when a member is promoted in a party
  * @param LocalUserId - id associated with this notification
@@ -847,6 +880,7 @@ PARTY_DECLARE_DELEGATETYPE(OnPartyMemberJoined);
 
 /**
  * Notification when party member data is updated
+ * Deprecated - Use OnPartyMemberDataReceivedConst
  * @param LocalUserId - id associated with this notification
  * @param PartyId - id associated with the party
  * @param MemberId - id of member that had updated data
@@ -856,7 +890,17 @@ DECLARE_MULTICAST_DELEGATE_FourParams(F_PREFIX(OnPartyMemberDataReceived), const
 PARTY_DECLARE_DELEGATETYPE(OnPartyMemberDataReceived);
 
 /**
- * Notification when an invite list has changed for a party
+ * Notification when party member data is updated
+ * @param LocalUserId - id associated with this notification
+ * @param PartyId - id associated with the party
+ * @param MemberId - id of member that had updated data
+ * @param PartyMemberData - party member data that was updated
+ */
+DECLARE_MULTICAST_DELEGATE_FourParams(F_PREFIX(OnPartyMemberDataReceivedConst), const FUniqueNetId& /*LocalUserId*/, const FOnlinePartyId& /*PartyId*/, const FUniqueNetId& /*MemberId*/, const FOnlinePartyData& /*PartyMemberData*/);
+PARTY_DECLARE_DELEGATETYPE(OnPartyMemberDataReceivedConst);
+
+/**
+ * Notification when an invite list has changed for a party the user is in
  * @param LocalUserId - user that is associated with this notification
  */
 DECLARE_MULTICAST_DELEGATE_OneParam(F_PREFIX(OnPartyInvitesChanged), const FUniqueNetId& /*LocalUserId*/);
@@ -864,29 +908,39 @@ PARTY_DECLARE_DELEGATETYPE(OnPartyInvitesChanged);
 
 /**
  * Notification when a request for an invite has been received
- * @param LocalUserId - id associated with this notification
- * @param PartyId - id associated with the party
- * @param SenderId - id of user that sent the invite
- * @param RequestForId - id of user that sender is requesting the invite for - invalid if the sender is requesting the invite
+ * @param LocalUserId id associated with this notification
+ * @param PartyId id associated with the party
+ * @param SenderId id of user that sent the invite
+ * @param RequestForId id of user that sender is requesting the invite for - invalid if the sender is requesting the invite
  */
 DECLARE_MULTICAST_DELEGATE_FourParams(F_PREFIX(OnPartyInviteRequestReceived), const FUniqueNetId& /*LocalUserId*/, const FOnlinePartyId& /*PartyId*/, const FUniqueNetId& /*SenderId*/, const FUniqueNetId& /*RequestForId*/);
 PARTY_DECLARE_DELEGATETYPE(OnPartyInviteRequestReceived);
 
 /**
  * Notification when a new invite is received
- * @param LocalUserId - id associated with this notification
- * @param PartyId - id associated with the party
- * @param SenderId - id of member that sent the invite
+ * @param LocalUserId id associated with this notification
+ * @param PartyId id associated with the party
+ * @param SenderId id of member that sent the invite
  */
 DECLARE_MULTICAST_DELEGATE_ThreeParams(F_PREFIX(OnPartyInviteReceived), const FUniqueNetId& /*LocalUserId*/, const FOnlinePartyId& /*PartyId*/, const FUniqueNetId& /*SenderId*/);
 PARTY_DECLARE_DELEGATETYPE(OnPartyInviteReceived);
 
+/**
+ * Notification when an invite has been removed
+ * @param LocalUserId id associated with this notification
+ * @param PartyId id associated with the party
+ * @param SenderId id of member that sent the invite
+ * @param Reason reason the invite has been removed
+ */
+DECLARE_MULTICAST_DELEGATE_FourParams(F_PREFIX(OnPartyInviteRemoved), const FUniqueNetId& /*LocalUserId*/, const FOnlinePartyId& /*PartyId*/, const FUniqueNetId& /*SenderId*/, EPartyInvitationRemovedReason /*Reason*/);
+PARTY_DECLARE_DELEGATETYPE(OnPartyInviteRemoved);
+
 /**
  * Notification when a new invite is received
- * @param LocalUserId - id associated with this notification
- * @param PartyId - id associated with the party
- * @param SenderId - id of member that sent the invite
- * @param bWasAccepted - whether or not the invite was accepted
+ * @param LocalUserId id associated with this notification
+ * @param PartyId id associated with the party
+ * @param SenderId id of member that sent the invite
+ * @param bWasAccepted whether or not the invite was accepted
  */
 DECLARE_MULTICAST_DELEGATE_FourParams(F_PREFIX(OnPartyInviteResponseReceived), const FUniqueNetId& /*LocalUserId*/, const FOnlinePartyId& /*PartyId*/, const FUniqueNetId& /*SenderId*/, const EInvitationResponse /*Response*/);
 PARTY_DECLARE_DELEGATETYPE(OnPartyInviteResponseReceived);
@@ -949,11 +1003,45 @@ PARTY_DECLARE_DELEGATETYPE(OnPartyAnalyticsEvent);
 DECLARE_MULTICAST_DELEGATE_OneParam(F_PREFIX(OnPartySystemStateChange), EPartySystemState /*NewState*/);
 PARTY_DECLARE_DELEGATETYPE(OnPartySystemStateChange);
 
+// Helper macro to aid in the deprecation of delegates that have non-const values
+#define DEFINE_ONLINE_DELEGATE_BASE_DEPRECATION_HELPER(DelegateName, NewDelegateName) \
+public: \
+	F##NewDelegateName DelegateName##Delegates; \
+public: \
+	virtual FDelegateHandle Add##DelegateName##Delegate_Handle(const F##NewDelegateName##Delegate& Delegate) \
+	{ \
+		DelegateName##Delegates.Add(Delegate); \
+		return Delegate.GetHandle(); \
+	} \
+	virtual void Clear##DelegateName##Delegate_Handle(FDelegateHandle& Handle) \
+	{ \
+		DelegateName##Delegates.Remove(Handle); \
+		Handle.Reset(); \
+	} \
+	virtual void Clear##DelegateName##Delegates(void* Object) \
+	{ \
+		DelegateName##Delegates.RemoveAll(Object); \
+	}
+
+#define DEFINE_ONLINE_DELEGATE_THREE_PARAM_DEPRECATION_HELPER(DelegateName, NewDelegateName, Param1Type, Param2Type, Param3Type) \
+DEFINE_ONLINE_DELEGATE_BASE_DEPRECATION_HELPER(DelegateName, NewDelegateName) \
+virtual void Trigger##DelegateName##Delegates(Param1Type Param1, Param2Type Param2, Param3Type Param3) \
+{ \
+	DelegateName##Delegates.Broadcast(Param1, Param2, Param3); \
+}
+
+#define DEFINE_ONLINE_DELEGATE_FOUR_PARAM_DEPRECATION_HELPER(DelegateName, DelegateNameConst, Param1Type, Param2Type, Param3Type, Param4Type) \
+DEFINE_ONLINE_DELEGATE_BASE_DEPRECATION_HELPER(DelegateName, DelegateNameConst) \
+virtual void Trigger##DelegateName##Delegates(Param1Type Param1, Param2Type Param2, Param3Type Param3, Param4Type Param4) \
+{ \
+	DelegateName##Delegates.Broadcast(Param1, Param2, Param3, Param4); \
+}
+
 /**
  * Interface definition for the online party services 
  * Allows for forming a party and communicating with party members
  */
-class IOnlinePartySystem
+class ONLINESUBSYSTEM_API IOnlinePartySystem
 {
 protected:
 	IOnlinePartySystem() {};
@@ -1116,7 +1204,8 @@ public:
 	 *
 	 * @return true if task was started
 	 */
-	virtual bool AcceptInvitation(const FUniqueNetId& LocalUserId, const FUniqueNetId& SenderId) = 0;
+	UE_DEPRECATED(4.23, "Use JoinParty instead of AcceptInvitation")
+	virtual bool AcceptInvitation(const FUniqueNetId& LocalUserId, const FUniqueNetId& SenderId) { return false; }
 
 	/**
 	 * Reject an invite to a party
@@ -1250,7 +1339,7 @@ public:
 	 *
 	 * @return party info or nullptr if not found
 	 */
-	virtual TSharedPtr<const FOnlineParty> GetParty(const FUniqueNetId& LocalUserId, const FOnlinePartyId& PartyId) const = 0;
+	virtual FOnlinePartyConstPtr GetParty(const FUniqueNetId& LocalUserId, const FOnlinePartyId& PartyId) const = 0;
 
 	/**
 	 * Get info associated with a party
@@ -1260,7 +1349,7 @@ public:
 	 *
 	 * @return party info or nullptr if not found
 	 */
-	virtual TSharedPtr<const FOnlineParty> GetParty(const FUniqueNetId& LocalUserId, const FOnlinePartyTypeId& PartyTypeId) const = 0;
+	virtual FOnlinePartyConstPtr GetParty(const FUniqueNetId& LocalUserId, const FOnlinePartyTypeId& PartyTypeId) const = 0;
 
 	/**
 	 * Get a party member by id
@@ -1271,7 +1360,7 @@ public:
 	 *
 	 * @return party member info or nullptr if not found
 	 */
-	virtual TSharedPtr<FOnlinePartyMember> GetPartyMember(const FUniqueNetId& LocalUserId, const FOnlinePartyId& PartyId, const FUniqueNetId& MemberId) const = 0;
+	virtual FOnlinePartyMemberConstPtr GetPartyMember(const FUniqueNetId& LocalUserId, const FOnlinePartyId& PartyId, const FUniqueNetId& MemberId) const = 0;
 
 	/**
 	 * Get current cached data associated with a party
@@ -1282,7 +1371,7 @@ public:
 	 *
 	 * @return party data or nullptr if not found
 	 */
-	virtual TSharedPtr<FOnlinePartyData> GetPartyData(const FUniqueNetId& LocalUserId, const FOnlinePartyId& PartyId) const = 0;
+	virtual FOnlinePartyDataConstPtr GetPartyData(const FUniqueNetId& LocalUserId, const FOnlinePartyId& PartyId) const = 0;
 
 	/**
 	 * Get current cached data associated with a party member
@@ -1294,7 +1383,7 @@ public:
 	 *
 	 * @return party member data or nullptr if not found
 	 */
-	virtual TSharedPtr<FOnlinePartyData> GetPartyMemberData(const FUniqueNetId& LocalUserId, const FOnlinePartyId& PartyId, const FUniqueNetId& MemberId) const = 0;
+	virtual FOnlinePartyDataConstPtr GetPartyMemberData(const FUniqueNetId& LocalUserId, const FOnlinePartyId& PartyId, const FUniqueNetId& MemberId) const = 0;
 
 	/**
 	 * Get the join info of the specified user and party type
@@ -1305,7 +1394,7 @@ public:
 	 *
 	 * @return shared pointer to the join info if the user is advertising for that party type
 	 */
-	virtual TSharedPtr<IOnlinePartyJoinInfo> GetAdvertisedParty(const FUniqueNetId& LocalUserId, const FUniqueNetId& UserId, const FOnlinePartyTypeId PartyTypeId) const = 0;
+	virtual IOnlinePartyJoinInfoConstPtr GetAdvertisedParty(const FUniqueNetId& LocalUserId, const FUniqueNetId& UserId, const FOnlinePartyTypeId PartyTypeId) const = 0;
 
 	/**
 	 * Get a list of currently joined parties for the user
@@ -1326,7 +1415,9 @@ public:
 	 *
 	 * @return true if entries found
 	 */
-	virtual bool GetPartyMembers(const FUniqueNetId& LocalUserId, const FOnlinePartyId& PartyId, TArray<TSharedRef<FOnlinePartyMember>>& OutPartyMembersArray) const = 0;
+	virtual bool GetPartyMembers(const FUniqueNetId& LocalUserId, const FOnlinePartyId& PartyId, TArray<FOnlinePartyMemberConstRef>& OutPartyMembersArray) const = 0;
+	UE_DEPRECATED(4.23, "Use GetPartyMembers that gives const FOnlinePartyMember")
+	virtual bool GetPartyMembers(const FUniqueNetId& LocalUserId, const FOnlinePartyId& PartyId, TArray<TSharedRef<FOnlinePartyMember>>& OutPartyMembersArray) const;
 
 	/**
 	 * Get a list of parties the user has been invited to
@@ -1336,7 +1427,9 @@ public:
 	 *
 	 * @return true if entries found
 	 */
-	virtual bool GetPendingInvites(const FUniqueNetId& LocalUserId, TArray<TSharedRef<IOnlinePartyJoinInfo>>& OutPendingInvitesArray) const = 0;
+	virtual bool GetPendingInvites(const FUniqueNetId& LocalUserId, TArray<IOnlinePartyJoinInfoConstRef>& OutPendingInvitesArray) const = 0;
+	UE_DEPRECATED(4.23, "Use GetPendingInvites that gives const IOnlinePartyJoinInfo")
+	virtual bool GetPendingInvites(const FUniqueNetId& LocalUserId, TArray<TSharedRef<IOnlinePartyJoinInfo>>& OutPendingInvitesArray) const;
 
 	/**
 	 * Get list of users requesting to join the party
@@ -1347,7 +1440,9 @@ public:
 	 *
 	 * @return true if entries found
 	 */
-	virtual bool GetPendingJoinRequests(const FUniqueNetId& LocalUserId, const FOnlinePartyId& PartyId, TArray<TSharedRef<IOnlinePartyPendingJoinRequestInfo>>& OutPendingJoinRequestArray) const = 0;
+	virtual bool GetPendingJoinRequests(const FUniqueNetId& LocalUserId, const FOnlinePartyId& PartyId, TArray<IOnlinePartyPendingJoinRequestInfoConstRef>& OutPendingJoinRequestArray) const = 0;
+	UE_DEPRECATED(4.23, "Use GetPendingJoinRequests that gives const IOnlinePartyPendingJoinRequestInfo")
+	virtual bool GetPendingJoinRequests(const FUniqueNetId& LocalUserId, const FOnlinePartyId& PartyId, TArray<TSharedRef<IOnlinePartyPendingJoinRequestInfo>>& OutPendingJoinRequestArray) const;
 
 	/**
 	 * Get list of users invited to a party that have not yet responded
@@ -1386,7 +1481,7 @@ public:
 	 * @param JoinInfoJson       - json blob to convert
 	 *
 	 */
-	virtual TSharedPtr<IOnlinePartyJoinInfo> MakeJoinInfoFromJson(const FString& JoinInfoJson) = 0;
+	virtual IOnlinePartyJoinInfoConstPtr MakeJoinInfoFromJson(const FString& JoinInfoJson) = 0;
 
 	/**
 	 * Creates a command line token from a IOnlinePartyJoinInfo object
@@ -1404,14 +1499,14 @@ public:
 	 *
 	 * return the new IOnlinePartyJoinInfo object
 	 */
-	virtual TSharedPtr<IOnlinePartyJoinInfo> MakeJoinInfoFromToken(const FString& Token) const = 0;
+	virtual IOnlinePartyJoinInfoConstPtr MakeJoinInfoFromToken(const FString& Token) const = 0;
 
 	/**
 	 * Checks to see if there is a pending command line invite and consumes it
 	 *
 	 * return the pending IOnlinePartyJoinInfo object
 	 */
-	virtual TSharedPtr<IOnlinePartyJoinInfo> ConsumePendingCommandLineInvite() = 0;
+	virtual IOnlinePartyJoinInfoConstPtr ConsumePendingCommandLineInvite() = 0;
 
 	/**
 	 * List of all subscribe-able notifications
@@ -1429,6 +1524,7 @@ public:
 	 * OnPartyInvitesChanged
 	 * OnPartyInviteRequestReceived
 	 * OnPartyInviteReceived
+	 * OnPartyInviteRemoved
 	 * OnPartyInviteResponseReceived
 	 * OnPartyJoinRequestReceived
 	 * OnPartyQueryJoinabilityReceived
@@ -1480,7 +1576,9 @@ public:
 	 * @param PartyId - id associated with the party
 	 * @param PartyConfig - party whose config was updated
 	 */
-	DEFINE_ONLINE_DELEGATE_THREE_PARAM(OnPartyConfigChanged, const FUniqueNetId& /*LocalUserId*/, const FOnlinePartyId& /*PartyId*/, const TSharedRef<FPartyConfiguration>& /*PartyConfig*/);
+	DEFINE_ONLINE_DELEGATE_THREE_PARAM_DEPRECATION_HELPER(OnPartyConfigChanged, OnPartyConfigChangedConst, const FUniqueNetId& /*LocalUserId*/, const FOnlinePartyId& /*PartyId*/, const FPartyConfiguration& /*PartyConfig*/);
+	UE_DEPRECATED(4.23, "Use OnPartyConfigChangedConst instead of OnPartyConfigChanged")
+	virtual FDelegateHandle AddOnPartyConfigChangedDelegate_Handle(const FOnPartyConfigChangedDelegate& Delegate);
 
 	/**
 	 * Notification when party data is updated
@@ -1488,7 +1586,9 @@ public:
 	 * @param PartyId - id associated with the party
 	 * @param PartyData - party data that was updated
 	 */
-	DEFINE_ONLINE_DELEGATE_THREE_PARAM(OnPartyDataReceived, const FUniqueNetId& /*LocalUserId*/, const FOnlinePartyId& /*PartyId*/, const TSharedRef<FOnlinePartyData>& /*PartyData*/);
+	DEFINE_ONLINE_DELEGATE_THREE_PARAM_DEPRECATION_HELPER(OnPartyDataReceived, OnPartyDataReceivedConst, const FUniqueNetId& /*LocalUserId*/, const FOnlinePartyId& /*PartyId*/, const FOnlinePartyData& /*PartyData*/);
+	UE_DEPRECATED(4.23, "Use OnPartyDataReceivedConst instead of OnPartyDataReceived")
+	virtual FDelegateHandle AddOnPartyDataReceivedDelegate_Handle(const FOnPartyDataReceivedDelegate& Delegate);
 
 	/**
 	* Notification when a member is promoted in a party
@@ -1522,7 +1622,9 @@ public:
 	 * @param MemberId - id of member that had updated data
 	 * @param PartyMemberData - party member data that was updated
 	 */
-	DEFINE_ONLINE_DELEGATE_FOUR_PARAM(OnPartyMemberDataReceived, const FUniqueNetId& /*LocalUserId*/, const FOnlinePartyId& /*PartyId*/, const FUniqueNetId& /*MemberId*/, const TSharedRef<FOnlinePartyData>& /*PartyMemberData*/);
+	DEFINE_ONLINE_DELEGATE_FOUR_PARAM_DEPRECATION_HELPER(OnPartyMemberDataReceived, OnPartyMemberDataReceivedConst, const FUniqueNetId& /*LocalUserId*/, const FOnlinePartyId& /*PartyId*/, const FUniqueNetId& /*MemberId*/, const FOnlinePartyData& /*PartyData*/);
+	UE_DEPRECATED(4.23, "Use OnPartyMemberDataReceivedConst instead of OnPartyMemberDataReceived")
+	virtual FDelegateHandle AddOnPartyMemberDataReceivedDelegate_Handle(const FOnPartyMemberDataReceivedDelegate& Delegate);
 
 	/**
 	 * Notification when an invite list has changed for a party
@@ -1540,13 +1642,22 @@ public:
 	DEFINE_ONLINE_DELEGATE_FOUR_PARAM(OnPartyInviteRequestReceived, const FUniqueNetId& /*LocalUserId*/, const FOnlinePartyId& /*PartyId*/, const FUniqueNetId& /*SenderId*/, const FUniqueNetId& /*RequestForId*/);
 
 	/**
-	* Notification when a new invite is received
-	* @param LocalUserId - id associated with this notification
-	* @param SenderId - id of member that sent the invite
-	* @param PartyId - id associated with the party
-	*/
+	 * Notification when a new invite is received
+	 * @param LocalUserId - id associated with this notification
+	 * @param PartyId - id associated with the party
+	 * @param SenderId - id of member that sent the invite
+	 */
 	DEFINE_ONLINE_DELEGATE_THREE_PARAM(OnPartyInviteReceived, const FUniqueNetId& /*LocalUserId*/, const FOnlinePartyId& /*PartyId*/, const FUniqueNetId& /*SenderId*/);
 
+	/**
+	 * Notification when an invite has been removed
+	 * @param LocalUserId id associated with this notification
+	 * @param PartyId id associated with the party
+	 * @param SenderId id of member that sent the invite
+	 * @param Reason reason the invitation was removed
+	 */
+	DEFINE_ONLINE_DELEGATE_FOUR_PARAM(OnPartyInviteRemoved, const FUniqueNetId& /*LocalUserId*/, const FOnlinePartyId& /*PartyId*/, const FUniqueNetId& /*SenderId*/, EPartyInvitationRemovedReason /*Reason*/);
+
 	/**
 	 * Notification when a new invite is received
 	 * @param LocalUserId - id associated with this notification
@@ -1620,7 +1731,7 @@ public:
 
 };
 
-enum class ECreatePartyCompletionResult
+enum class ECreatePartyCompletionResult : int8
 {
 	UnknownClientFailure = -100,
 	AlreadyInPartyOfSpecifiedType,
@@ -1633,7 +1744,7 @@ enum class ECreatePartyCompletionResult
 	Succeeded = 1
 };
 
-enum class EJoinPartyCompletionResult
+enum class EJoinPartyCompletionResult : int8
 {
 	/** Unspecified error.  No message sent to party leader. */
 	UnknownClientFailure = -100,
@@ -1683,7 +1794,7 @@ enum class EJoinPartyCompletionResult
 	Succeeded = 1
 };
 
-enum class ELeavePartyCompletionResult
+enum class ELeavePartyCompletionResult : int8
 {
 	/** Unspecified error.  No message sent. */
 	UnknownClientFailure = -100,
@@ -1712,7 +1823,7 @@ enum class ELeavePartyCompletionResult
 	Succeeded = 1
 };
 
-enum class EUpdateConfigCompletionResult
+enum class EUpdateConfigCompletionResult : int8
 {
 	UnknownClientFailure = -100,
 	UnknownParty,
@@ -1725,7 +1836,7 @@ enum class EUpdateConfigCompletionResult
 	Succeeded = 1
 };
 
-enum class ERequestPartyInvitationCompletionResult
+enum class ERequestPartyInvitationCompletionResult : int8
 {
 	NotLoggedIn = -100,
 	InvitePending,
@@ -1736,7 +1847,7 @@ enum class ERequestPartyInvitationCompletionResult
 	Succeeded = 1
 };
 
-enum class ESendPartyInvitationCompletionResult
+enum class ESendPartyInvitationCompletionResult : int8
 {
 	NotLoggedIn = -100,
 	InvitePending,
@@ -1747,7 +1858,7 @@ enum class ESendPartyInvitationCompletionResult
 	Succeeded = 1
 };
 
-enum class EAcceptPartyInvitationCompletionResult
+enum class EAcceptPartyInvitationCompletionResult : int8
 {
 	NotLoggedIn = -100,
 	InvitePending,
@@ -1758,7 +1869,7 @@ enum class EAcceptPartyInvitationCompletionResult
 	Succeeded = 1
 };
 
-enum class ERejectPartyInvitationCompletionResult
+enum class ERejectPartyInvitationCompletionResult : int8
 {
 	NotLoggedIn = -100,
 	InvitePending,
@@ -1769,7 +1880,7 @@ enum class ERejectPartyInvitationCompletionResult
 	Succeeded = 1
 };
 
-enum class EKickMemberCompletionResult
+enum class EKickMemberCompletionResult : int8
 {
 	UnknownClientFailure = -100,
 	UnknownParty,
@@ -1783,7 +1894,7 @@ enum class EKickMemberCompletionResult
 	Succeeded = 1
 };
 
-enum class EPromoteMemberCompletionResult
+enum class EPromoteMemberCompletionResult : int8
 {
 	UnknownClientFailure = -100,
 	UnknownServiceFailure,
@@ -1800,7 +1911,7 @@ enum class EPromoteMemberCompletionResult
 	Succeeded = 1
 };
 
-enum class EInvitationResponse
+enum class EInvitationResponse : uint8
 {
 	UnknownFailure,
 	BadBuild,
