Fix parameter comment

diff --git a/Engine/Plugins/Online/OnlineSubsystem/Source/Public/Interfaces/OnlinePresenceInterface.h b/Engine/Plugins/Online/OnlineSubsystem/Source/Public/Interfaces/OnlinePresenceInterface.h
--- a/Engine/Plugins/Online/OnlineSubsystem/Source/Public/Interfaces/OnlinePresenceInterface.h
+++ b/Engine/Plugins/Online/OnlineSubsystem/Source/Public/Interfaces/OnlinePresenceInterface.h
@@ -237,7 +237,7 @@ public:
  * Delegate executed when new presence data is available for a user.
  *
  * @param UserId The unique id of the user whose presence was received.
- * @param Presence The unique id of the user whose presence was received.
+ * @param Presence The received presence
  */
 DECLARE_MULTICAST_DELEGATE_TwoParams(FOnPresenceReceived, const class FUniqueNetId& /*UserId*/, const TSharedRef<FOnlineUserPresence>& /*Presence*/);
 typedef FOnPresenceReceived::FDelegate FOnPresenceReceivedDelegate;
