Add more invitation removed reasons

diff --git a/Engine/Plugins/Online/OnlineSubsystem/Source/Private/OnlinePartyInterface.cpp b/Engine/Plugins/Online/OnlineSubsystem/Source/Private/OnlinePartyInterface.cpp
--- a/Engine/Plugins/Online/OnlineSubsystem/Source/Private/OnlinePartyInterface.cpp
+++ b/Engine/Plugins/Online/OnlineSubsystem/Source/Private/OnlinePartyInterface.cpp
@@ -321,6 +321,38 @@ const TCHAR* ToString(const EMemberExitedReason Value)
 	return TEXT("Unknown"); // Same as EMemberExitedReason::Unknown, which is ok because it is only used when we do not have enough information
 }
 
+const TCHAR* ToString(const EPartyInvitationRemovedReason Value)
+{
+	switch (Value)
+	{
+	case EPartyInvitationRemovedReason::Unknown:
+	{
+		return TEXT("Unknown");
+	}
+	case EPartyInvitationRemovedReason::Accepted:
+	{
+		return TEXT("Accepted");
+	}
+	case EPartyInvitationRemovedReason::Declined:
+	{
+		return TEXT("Declined");
+	}
+	case EPartyInvitationRemovedReason::Cleared:
+	{
+		return TEXT("Cleared");
+	}
+	case EPartyInvitationRemovedReason::Expired:
+	{
+		return TEXT("Expired");
+	}
+	case EPartyInvitationRemovedReason::Invalidated:
+	{
+		return TEXT("Invalidated");
+	}
+	}
+	return TEXT("Unknown"); // Same as EMemberExitedReason::Unknown, which is ok because it is only used when we do not have enough information
+}
+
 const TCHAR* ToString(const ECreatePartyCompletionResult Value)
 {
 	switch (Value)
diff --git a/Engine/Plugins/Online/OnlineSubsystem/Source/Public/Interfaces/OnlinePartyInterface.h b/Engine/Plugins/Online/OnlineSubsystem/Source/Public/Interfaces/OnlinePartyInterface.h
--- a/Engine/Plugins/Online/OnlineSubsystem/Source/Public/Interfaces/OnlinePartyInterface.h
+++ b/Engine/Plugins/Online/OnlineSubsystem/Source/Public/Interfaces/OnlinePartyInterface.h
@@ -610,12 +610,16 @@ enum class EPartyInvitationRemovedReason : uint8
 {
 	/** Unknown or undefined reason */
 	Unknown,
+	/** User accepted the invitation */
+	Accepted,
+	/** User declined the invitation */
+	Declined,
+	/** ClearInvitations was called, the invitation should no longer be displayed */
+	Cleared,
 	/** Expired */
 	Expired,
 	/** Became invalid (for example, party was destroyed) */
 	Invalidated,
-	/** User declined the invitation */
-	Declined,
 };
 
 /** Recipient information for SendInvitation */
@@ -1927,6 +1931,8 @@ ONLINESUBSYSTEM_API EPartyState EPartyStateFromString(const TCHAR* Value);
 /** @return the stringified version of the enum passed in */
 ONLINESUBSYSTEM_API const TCHAR* ToString(const EMemberExitedReason Value);
 /** @return the stringified version of the enum passed in */
+ONLINESUBSYSTEM_API const TCHAR* ToString(const EPartyInvitationRemovedReason Value);
+/** @return the stringified version of the enum passed in */
 ONLINESUBSYSTEM_API const TCHAR* ToString(const ECreatePartyCompletionResult Value);
 /** @return the stringified version of the enum passed in */
 ONLINESUBSYSTEM_API const TCHAR* ToString(const ESendPartyInvitationCompletionResult Value);
