Add OnMemberAttributeChanged event to FOnlinePartyMember

diff --git a/Engine/Plugins/Online/OnlineSubsystem/Source/Public/Interfaces/OnlinePartyInterface.h b/Engine/Plugins/Online/OnlineSubsystem/Source/Public/Interfaces/OnlinePartyInterface.h
--- a/Engine/Plugins/Online/OnlineSubsystem/Source/Public/Interfaces/OnlinePartyInterface.h
+++ b/Engine/Plugins/Online/OnlineSubsystem/Source/Public/Interfaces/OnlinePartyInterface.h
@@ -58,6 +58,17 @@ public:
 	EMemberConnectionStatus MemberConnectionStatus = EMemberConnectionStatus::Uninitialized;
 	EMemberConnectionStatus PreviousMemberConnectionStatus = EMemberConnectionStatus::Uninitialized;
 
+	/**
+	 * Event when a party member's attribute has changed
+	 * @see FOnlineUser::GetUserAttribute
+	 * @param ChangedUserId id associated with this notification
+	 * @param Attribute attribute that changed
+	 * @param NewValue the new value for the attribute
+	 * @param PreviousValue the previous value for the attribute
+	 */
+	DECLARE_EVENT_FourParams(FOnlinePartyMember, FOnMemberAttributeChanged, const FUniqueNetId& /*ChangedUserId*/, const FString& /*Attribute*/, const FString& /*NewValue*/, const FString& /*OldValue*/);
+	FOnMemberAttributeChanged& OnMemberAttributeChanged() const { return OnMemberAttributeChangedEvent; }
+
 	/**
 	 * Event when a party member's connection status has changed
 	 * @param ChangedUserId - id associated with this notification
@@ -80,6 +91,8 @@ public:
 private:
 	/** Event fired when connection status changes */
 	mutable FOnMemberConnectionStatusChanged OnMemberConnectionStatusChangedEvent;
+	/** Event fired when an attribute changes */
+	mutable FOnMemberAttributeChanged OnMemberAttributeChangedEvent;
 };
 
 typedef TSharedRef<const FOnlinePartyMember> FOnlinePartyMemberConstRef;
@@ -943,7 +956,7 @@ PARTY_DECLARE_DELEGATETYPE(OnPartyInviteRemoved);
  * Notification when a new invite is received
  * @param LocalUserId id associated with this notification
  * @param PartyId id associated with the party
- * @param SenderId id of member that sent the invite
+ * @param SenderId id of user that sent the invite
  * @param bWasAccepted whether or not the invite was accepted
  */
 DECLARE_MULTICAST_DELEGATE_FourParams(F_PREFIX(OnPartyInviteResponseReceived), const FUniqueNetId& /*LocalUserId*/, const FOnlinePartyId& /*PartyId*/, const FUniqueNetId& /*SenderId*/, const EInvitationResponse /*Response*/);
@@ -1663,11 +1676,11 @@ public:
 	DEFINE_ONLINE_DELEGATE_FOUR_PARAM(OnPartyInviteRemoved, const FUniqueNetId& /*LocalUserId*/, const FOnlinePartyId& /*PartyId*/, const FUniqueNetId& /*SenderId*/, EPartyInvitationRemovedReason /*Reason*/);
 
 	/**
-	 * Notification when a new invite is received
+	 * Notification when an invitation response is received
 	 * @param LocalUserId - id associated with this notification
 	 * @param PartyId - id associated with the party
-	 * @param SenderId - id of member that sent the invite
-	 * @param bWasAccepted - true is the invite was accepted
+	 * @param SenderId - id of user that responded to an invite
+	 * @param Response - the response
 	 */
 	DEFINE_ONLINE_DELEGATE_FOUR_PARAM(OnPartyInviteResponseReceived, const FUniqueNetId& /*LocalUserId*/, const FOnlinePartyId& /*PartyId*/, const FUniqueNetId& /*SenderId*/, const EInvitationResponse /*Response*/);
 
