Add LeaveParty overload to optionally synchronize leave across clients

diff --git a/Engine/Plugins/Online/OnlineSubsystem/Source/Public/Interfaces/OnlinePartyInterface.h b/Engine/Plugins/Online/OnlineSubsystem/Source/Public/Interfaces/OnlinePartyInterface.h
--- a/Engine/Plugins/Online/OnlineSubsystem/Source/Public/Interfaces/OnlinePartyInterface.h
+++ b/Engine/Plugins/Online/OnlineSubsystem/Source/Public/Interfaces/OnlinePartyInterface.h
@@ -1163,6 +1163,19 @@ public:
 	 */
 	virtual bool LeaveParty(const FUniqueNetId& LocalUserId, const FOnlinePartyId& PartyId, const FOnLeavePartyComplete& Delegate = FOnLeavePartyComplete()) = 0;
 
+	/**
+	 * Leave an existing party
+	 * All existing party members notified of member leaving (see FOnPartyMemberLeft)
+	 *
+	 * @param LocalUserId - user making the request
+	 * @param PartyId - id of an existing party
+	 * @param bSynchronizeLeave - Whether we synchronize the leave with remote server/clients or only do a local cleanup
+	 * @param Delegate - called on completion
+	 *
+	 * @return true if task was started
+	 */
+	virtual bool LeaveParty(const FUniqueNetId& LocalUserId, const FOnlinePartyId& PartyId, bool bSynchronizeLeave, const FOnLeavePartyComplete& Delegate = FOnLeavePartyComplete()) = 0;
+
 	/**
 	* Approve a request to join a party
 	*
