#OnlineSubsystem - Fixed uninitialized memory in OnlineError
- Fixes cases where errors were sometimes treated as success

diff --git a/Engine/Plugins/Online/OnlineSubsystem/Source/Private/OnlineError.cpp b/Engine/Plugins/Online/OnlineSubsystem/Source/Private/OnlineError.cpp
--- a/Engine/Plugins/Online/OnlineSubsystem/Source/Private/OnlineError.cpp
+++ b/Engine/Plugins/Online/OnlineSubsystem/Source/Private/OnlineError.cpp
@@ -18,6 +18,7 @@ FOnlineError::FOnlineError(const TCHAR* const ErrorCodeIn)
 
 FOnlineError::FOnlineError(const int32 ErrorCodeIn)
 	: bSucceeded(false)
+	, Result(EOnlineErrorResult::FailExtended)
 {
 	SetFromErrorCode(ErrorCodeIn);
 }
@@ -96,37 +97,43 @@ void FOnlineError::SetFromErrorCode(EOnlineErrorResult InResult, const FString&
 
 FOnlineError::FOnlineError(bool bSucceededIn)
 	: bSucceeded(bSucceededIn)
+	, Result(bSucceededIn ? EOnlineErrorResult::Success : EOnlineErrorResult::Unknown)
 {
 }
 
 FOnlineError::FOnlineError(const FString& ErrorCodeIn)
 	: bSucceeded(false)
+	, Result(EOnlineErrorResult::FailExtended)
 {
 	SetFromErrorCode(ErrorCodeIn);
 }
 
 FOnlineError::FOnlineError(FString&& ErrorCodeIn)
 	: bSucceeded(false)
+	, Result(EOnlineErrorResult::FailExtended)
 {
 	SetFromErrorCode(MoveTemp(ErrorCodeIn));
 }
 
+FOnlineError::FOnlineError(const FText& ErrorMessageIn)
+	: bSucceeded(false)
+	, Result(EOnlineErrorResult::FailExtended)
+{
+	SetFromErrorMessage(ErrorMessageIn);
+}
+
 void FOnlineError::SetFromErrorCode(const FString& ErrorCodeIn)
 {
 	ErrorCode = ErrorCodeIn;
 	ErrorRaw = ErrorCodeIn;
+	Result = EOnlineErrorResult::FailExtended;
 }
 
 void FOnlineError::SetFromErrorCode(FString&& ErrorCodeIn)
 {
 	ErrorCode = MoveTemp(ErrorCodeIn);
 	ErrorRaw = ErrorCode;
-}
-
-FOnlineError::FOnlineError(const FText& ErrorMessageIn)
-	: bSucceeded(false)
-{
-	SetFromErrorMessage(ErrorMessageIn);
+	Result = EOnlineErrorResult::FailExtended;
 }
 
 void FOnlineError::SetFromErrorMessage(const FText& ErrorMessageIn)
@@ -134,6 +141,7 @@ void FOnlineError::SetFromErrorMessage(const FText& ErrorMessageIn)
 	ErrorMessage = ErrorMessageIn;
 	ErrorCode = FTextInspector::GetKey(ErrorMessageIn).Get(GenericErrorCode);
 	ErrorRaw = ErrorMessageIn.ToString();
+	Result = EOnlineErrorResult::FailExtended;
 }
 
 FString FOnlineError::ToLogString() const
