Fixed initialization of the FOnlineError bool constructor so that it does not set the Result member to a success code. This ensures that legacy code that changes bSucceeded from true to false after construction will not receive incorrect results from WasSuccessful().

diff --git a/Engine/Plugins/Online/OnlineSubsystem/Source/Private/OnlineError.cpp b/Engine/Plugins/Online/OnlineSubsystem/Source/Private/OnlineError.cpp
--- a/Engine/Plugins/Online/OnlineSubsystem/Source/Private/OnlineError.cpp
+++ b/Engine/Plugins/Online/OnlineSubsystem/Source/Private/OnlineError.cpp
@@ -97,7 +97,7 @@ void FOnlineError::SetFromErrorCode(EOnlineErrorResult InResult, const FString&
 
 FOnlineError::FOnlineError(bool bSucceededIn)
 	: bSucceeded(bSucceededIn)
-	, Result(bSucceededIn ? EOnlineErrorResult::Success : EOnlineErrorResult::Unknown)
+	, Result(EOnlineErrorResult::Unknown)
 {
 }
 
