Fix up IP Serialization so that it supports IPv4/IPv6. Fixes bad travels on LAN due to invalid address.

diff --git a/Engine/Plugins/Online/OnlineSubsystem/Source/Public/NboSerializer.h b/Engine/Plugins/Online/OnlineSubsystem/Source/Public/NboSerializer.h
--- a/Engine/Plugins/Online/OnlineSubsystem/Source/Public/NboSerializer.h
+++ b/Engine/Plugins/Online/OnlineSubsystem/Source/Public/NboSerializer.h
@@ -417,11 +417,10 @@ public:
 	 */
 	friend inline FNboSerializeToBuffer& operator<<(FNboSerializeToBuffer& Ar,const FInternetAddr& Addr)
 	{
-		uint32 OutIp;
-		Addr.GetIp(OutIp);
-		Ar << OutIp;
-		int32 OutPort = Addr.GetPort();
-		Ar << OutPort;
+		TArray<uint8> RawAddressBytes = Addr.GetRawIp();
+		Ar << RawAddressBytes.Num();
+		Ar.WriteBinary(RawAddressBytes.GetData(), RawAddressBytes.Num());
+		Ar << Addr.GetPort();
 		return Ar;
 	}
 
@@ -812,9 +811,11 @@ public:
 	 */
 	friend inline FNboSerializeFromBuffer& operator>>(FNboSerializeFromBuffer& Ar,FInternetAddr& Addr)
 	{
-		uint32 InIp = 0;
-		Ar >> InIp;
-		Addr.SetIp(InIp);
+		TArray<uint8> RawAddressBytes;
+		uint32 AddressSize = 0;
+		Ar >> AddressSize;
+		Ar.ReadBinaryArray(RawAddressBytes, AddressSize);
+		Addr.SetRawIp(RawAddressBytes);
 		int32 InPort = 0;
 		Ar >> InPort;
 		Addr.SetPort(InPort);
@@ -833,6 +834,30 @@ public:
 		return Ar;
 	}
 
+	/**
+	 * Reads a blob of data from the buffer into an array
+	 * Prevents the necessity for additional allocation.
+	 *
+	 * @param OutArray the destination array
+	 * @param NumToRead the size of the blob to read
+	 */
+	void ReadBinaryArray(TArray<uint8>& OutArray, uint32 NumToRead)
+	{
+		OutArray.Reserve(NumToRead);
+		if (!HasOverflow() && CurrentOffset + (int32)NumToRead <= NumBytes)
+		{
+			for (uint32 i = 0; i < NumToRead; ++i)
+			{
+				OutArray.Add(Data[CurrentOffset + i]);
+			}
+			CurrentOffset += NumToRead;
+		}
+		else
+		{
+			bHasOverflowed = true;
+		}
+	}
+
 	/**
 	 * Reads a blob of data from the buffer
 	 *
