Add a few party string-to-enum functions

diff --git a/Engine/Plugins/Online/OnlineSubsystem/Source/Private/OnlinePartyInterface.cpp b/Engine/Plugins/Online/OnlineSubsystem/Source/Private/OnlinePartyInterface.cpp
--- a/Engine/Plugins/Online/OnlineSubsystem/Source/Private/OnlinePartyInterface.cpp
+++ b/Engine/Plugins/Online/OnlineSubsystem/Source/Private/OnlinePartyInterface.cpp
@@ -156,3 +156,622 @@ bool FPartyConfiguration::operator!=(const FPartyConfiguration& Other) const
 {
 	return !operator==(Other);
 }
+
+const TCHAR* ToString(const EPartyState Value)
+{
+	switch (Value)
+	{
+	case EPartyState::None:
+	{
+		return TEXT("None");
+	}
+	case EPartyState::CreatePending:
+	{
+		return TEXT("CreatePending");
+	}
+	case EPartyState::JoinPending:
+	{
+		return TEXT("JoinPending");
+	}
+	case EPartyState::LeavePending:
+	{
+		return TEXT("LeavePending");
+	}
+	case EPartyState::Active:
+	{
+		return TEXT("Active");
+	}
+	case EPartyState::Disconnected:
+	{
+		return TEXT("Disconnected");
+	}
+	case EPartyState::CleanUp:
+	{
+		return TEXT("CleanUp");
+	}
+	}
+	return TEXT("Unknown");
+}
+
+EPartyState EPartyStateFromString(const TCHAR* Value)
+{
+	if (FCString::Stricmp(Value, TEXT("CreatePending")) == 0)
+	{
+		return EPartyState::CreatePending;
+	}
+	else if (FCString::Stricmp(Value, TEXT("JoinPending")) == 0)
+	{
+		return EPartyState::JoinPending;
+	}
+	else if (FCString::Stricmp(Value, TEXT("LeavePending")) == 0)
+	{
+		return EPartyState::LeavePending;
+	}
+	else if (FCString::Stricmp(Value, TEXT("Active")) == 0)
+	{
+		return EPartyState::Active;
+	}
+	else if (FCString::Stricmp(Value, TEXT("Disconnected")) == 0)
+	{
+		return EPartyState::Disconnected;
+	}
+	else if (FCString::Stricmp(Value, TEXT("CleanUp")) == 0)
+	{
+		return EPartyState::CleanUp;
+	}
+	return EPartyState::None;
+}
+
+const TCHAR* ToString(const EMemberExitedReason Value)
+{
+	switch (Value)
+	{
+	case EMemberExitedReason::Unknown:
+	{
+		return TEXT("Unknown");
+	}
+	case EMemberExitedReason::Left:
+	{
+		return TEXT("Left");
+	}
+	case EMemberExitedReason::Removed:
+	{
+		return TEXT("Removed");
+	}
+	case EMemberExitedReason::Kicked:
+	{
+		return TEXT("Kicked");
+	}
+	}
+	return TEXT("Unknown"); // Same as EMemberExitedReason::Unknown, which is ok because it is only used when we do not have enough information
+}
+
+const TCHAR* ToString(const ECreatePartyCompletionResult Value)
+{
+	switch (Value)
+	{
+	case ECreatePartyCompletionResult::UnknownClientFailure:
+	{
+		return TEXT("UnknownClientFailure");
+	}
+	case ECreatePartyCompletionResult::AlreadyCreatingParty:
+	{
+		return TEXT("AlreadyCreatingParty");
+	}
+	case ECreatePartyCompletionResult::AlreadyInParty:
+	{
+		return TEXT("AlreadyInParty");
+	}
+	case ECreatePartyCompletionResult::FailedToCreateMucRoom:
+	{
+		return TEXT("FailedToCreateMucRoom");
+	}
+	case ECreatePartyCompletionResult::NoResponse:
+	{
+		return TEXT("NoResponse");
+	}
+	case ECreatePartyCompletionResult::LoggedOut:
+	{
+		return TEXT("LoggedOut");
+	}
+	case ECreatePartyCompletionResult::UnknownInternalFailure:
+	{
+		return TEXT("UnknownInternalFailure");
+	}
+	case ECreatePartyCompletionResult::Succeeded:
+	{
+		return TEXT("Succeeded");
+	}
+	}
+	return TEXT("Unknown");
+}
+
+const TCHAR* ToString(const ESendPartyInvitationCompletionResult Value)
+{
+	switch (Value)
+	{
+	case ESendPartyInvitationCompletionResult::NotLoggedIn:
+	{
+		return TEXT("NotLoggedIn");
+	}
+	case ESendPartyInvitationCompletionResult::InvitePending:
+	{
+		return TEXT("InvitePending");
+	}
+	case ESendPartyInvitationCompletionResult::AlreadyInParty:
+	{
+		return TEXT("AlreadyInParty");
+	}
+	case ESendPartyInvitationCompletionResult::PartyFull:
+	{
+		return TEXT("PartyFull");
+	}
+	case ESendPartyInvitationCompletionResult::NoPermission:
+	{
+		return TEXT("NoPermission");
+	}
+	case ESendPartyInvitationCompletionResult::UnknownInternalFailure:
+	{
+		return TEXT("UnknownInternalFailure");
+	}
+	case ESendPartyInvitationCompletionResult::Succeeded:
+	{
+		return TEXT("Succeeded");
+	}
+	}
+	return TEXT("Unknown");
+}
+
+const TCHAR* ToString(const EJoinPartyCompletionResult Value)
+{
+	switch (Value)
+	{
+	case EJoinPartyCompletionResult::UnknownClientFailure:
+	{
+		return TEXT("UnknownClientFailure");
+	}
+	case EJoinPartyCompletionResult::BadBuild:
+	{
+		return TEXT("BadBuild");
+	}
+	case EJoinPartyCompletionResult::InvalidAccessKey:
+	{
+		return TEXT("InvalidAccessKey");
+	}
+	case EJoinPartyCompletionResult::AlreadyInLeadersJoiningList:
+	{
+		return TEXT("AlreadyInLeadersJoiningList");
+	}
+	case EJoinPartyCompletionResult::AlreadyInLeadersPartyRoster:
+	{
+		return TEXT("AlreadyInLeadersPartyRoster");
+	}
+	case EJoinPartyCompletionResult::NoSpace:
+	{
+		return TEXT("NoSpace");
+	}
+	case EJoinPartyCompletionResult::NotApproved:
+	{
+		return TEXT("NotApproved");
+	}
+	case EJoinPartyCompletionResult::RequesteeNotMember:
+	{
+		return TEXT("RequesteeNotMember");
+	}
+	case EJoinPartyCompletionResult::RequesteeNotLeader:
+	{
+		return TEXT("RequesteeNotLeader");
+	}
+	case EJoinPartyCompletionResult::NoResponse:
+	{
+		return TEXT("NoResponse");
+	}
+	case EJoinPartyCompletionResult::LoggedOut:
+	{
+		return TEXT("LoggedOut");
+	}
+	case EJoinPartyCompletionResult::UnableToRejoin:
+	{
+		return TEXT("UnableToRejoin");
+	}
+	case EJoinPartyCompletionResult::IncompatiblePlatform:
+	{
+		return TEXT("IncompatiblePlatform");
+	}
+	case EJoinPartyCompletionResult::AlreadyJoiningParty:
+	{
+		return TEXT("AlreadyJoiningParty");
+	}
+	case EJoinPartyCompletionResult::AlreadyInParty:
+	{
+		return TEXT("AlreadyInParty");
+	}
+	case EJoinPartyCompletionResult::JoinInfoInvalid:
+	{
+		return TEXT("JoinInfoInvalid");
+	}
+	case EJoinPartyCompletionResult::AlreadyInPartyOfSpecifiedType:
+	{
+		return TEXT("AlreadyInPartyOfSpecifiedType");
+	}
+	case EJoinPartyCompletionResult::MessagingFailure:
+	{
+		return TEXT("MessagingFailure");
+	}
+	case EJoinPartyCompletionResult::GameSpecificReason:
+	{
+		return TEXT("GameSpecificReason");
+	}
+	case EJoinPartyCompletionResult::Succeeded:
+	{
+		return TEXT("Succeeded");
+	}
+	case EJoinPartyCompletionResult::UnknownInternalFailure:
+	{
+		return TEXT("DeprecatedUnknownInternalFailure");
+	}
+	}
+	return TEXT("Unknown");
+}
+
+const TCHAR* ToString(const ELeavePartyCompletionResult Value)
+{
+	switch (Value)
+	{
+	case ELeavePartyCompletionResult::UnknownClientFailure:
+	{
+		return TEXT("UnknownClientFailure");
+	}
+	case ELeavePartyCompletionResult::NoResponse:
+	{
+		return TEXT("NoResponse");
+	}
+	case ELeavePartyCompletionResult::LoggedOut:
+	{
+		return TEXT("LoggedOut");
+	}
+	case ELeavePartyCompletionResult::UnknownParty:
+	{
+		return TEXT("UnknownParty");
+	}
+	case ELeavePartyCompletionResult::LeavePending:
+	{
+		return TEXT("LeavePending");
+	}
+	case ELeavePartyCompletionResult::Succeeded:
+	{
+		return TEXT("Succeeded");
+	}
+	case ELeavePartyCompletionResult::UnknownLocalUser:
+	{
+		return TEXT("DeprecatedUnknownLocalUser");
+	}
+	case ELeavePartyCompletionResult::NotMember:
+	{
+		return TEXT("DeprecatedNotMember");
+	}
+	case ELeavePartyCompletionResult::MessagingFailure:
+	{
+		return TEXT("DeprecatedMessagingFailure");
+	}
+	case ELeavePartyCompletionResult::UnknownTransportFailure:
+	{
+		return TEXT("DeprecatedUnknownTransportFailure");
+	}
+	case ELeavePartyCompletionResult::UnknownInternalFailure:
+	{
+		return TEXT("DeprecatedUnknownInternalFailure");
+	}
+	}
+	return TEXT("Unknown");
+}
+
+const TCHAR* ToString(const EUpdateConfigCompletionResult Value)
+{
+	switch (Value)
+	{
+	case EUpdateConfigCompletionResult::UnknownClientFailure:
+	{
+		return TEXT("UnknownClientFailure");
+	}
+	case EUpdateConfigCompletionResult::UnknownParty:
+	{
+		return TEXT("UnknownParty");
+	}
+	case EUpdateConfigCompletionResult::LocalMemberNotMember:
+	{
+		return TEXT("LocalMemberNotMember");
+	}
+	case EUpdateConfigCompletionResult::LocalMemberNotLeader:
+	{
+		return TEXT("LocalMemberNotLeader");
+	}
+	case EUpdateConfigCompletionResult::RemoteMemberNotMember:
+	{
+		return TEXT("RemoteMemberNotMember");
+	}
+	case EUpdateConfigCompletionResult::MessagingFailure:
+	{
+		return TEXT("MessagingFailure");
+	}
+	case EUpdateConfigCompletionResult::NoResponse:
+	{
+		return TEXT("NoResponse");
+	}
+	case EUpdateConfigCompletionResult::UnknownInternalFailure:
+	{
+		return TEXT("UnknownInternalFailure");
+	}
+	case EUpdateConfigCompletionResult::Succeeded:
+	{
+		return TEXT("Succeeded");
+	}
+	}
+	return TEXT("Unknown");
+}
+
+const TCHAR* ToString(const EKickMemberCompletionResult Value)
+{
+	switch (Value)
+	{
+	case EKickMemberCompletionResult::UnknownClientFailure:
+	{
+		return TEXT("UnknownClientFailure");
+	}
+	case EKickMemberCompletionResult::UnknownParty:
+	{
+		return TEXT("UnknownParty");
+	}
+	case EKickMemberCompletionResult::LocalMemberNotMember:
+	{
+		return TEXT("LocalMemberNotMember");
+	}
+	case EKickMemberCompletionResult::LocalMemberNotLeader:
+	{
+		return TEXT("LocalMemberNotLeader");
+	}
+	case EKickMemberCompletionResult::RemoteMemberNotMember:
+	{
+		return TEXT("RemoteMemberNotMember");
+	}
+	case EKickMemberCompletionResult::MessagingFailure:
+	{
+		return TEXT("MessagingFailure");
+	}
+	case EKickMemberCompletionResult::NoResponse:
+	{
+		return TEXT("NoResponse");
+	}
+	case EKickMemberCompletionResult::LoggedOut:
+	{
+		return TEXT("LoggedOut");
+	}
+	case EKickMemberCompletionResult::UnknownInternalFailure:
+	{
+		return TEXT("UnknownInternalFailure");
+	}
+	case EKickMemberCompletionResult::Succeeded:
+	{
+		return TEXT("Succeeded");
+	}
+	}
+	return TEXT("Unknown");
+}
+
+const TCHAR* ToString(const EPromoteMemberCompletionResult Value)
+{
+	switch (Value)
+	{
+	case EPromoteMemberCompletionResult::UnknownClientFailure:
+	{
+		return TEXT("UnknownClientFailure");
+	}
+	case EPromoteMemberCompletionResult::UnknownParty:
+	{
+		return TEXT("UnknownParty");
+	}
+	case EPromoteMemberCompletionResult::LocalMemberNotMember:
+	{
+		return TEXT("LocalMemberNotMember");
+	}
+	case EPromoteMemberCompletionResult::LocalMemberNotLeader:
+	{
+		return TEXT("LocalMemberNotLeader");
+	}
+	case EPromoteMemberCompletionResult::TargetIsSelf:
+	{
+		return TEXT("TargetIsSelf");
+	}
+	case EPromoteMemberCompletionResult::TargetNotMember:
+	{
+		return TEXT("TargetNotMember");
+	}
+	case EPromoteMemberCompletionResult::MessagingFailure:
+	{
+		return TEXT("MessagingFailure");
+	}
+	case EPromoteMemberCompletionResult::NoResponse:
+	{
+		return TEXT("NoResponse");
+	}
+	case EPromoteMemberCompletionResult::LoggedOut:
+	{
+		return TEXT("LoggedOut");
+	}
+	case EPromoteMemberCompletionResult::UnknownInternalFailure:
+	{
+		return TEXT("UnknownInternalFailure");
+	}
+	case EPromoteMemberCompletionResult::Succeeded:
+	{
+		return TEXT("Succeeded");
+	}
+	}
+	return TEXT("Unknown");
+}
+
+const TCHAR* ToString(const EInvitationResponse Value)
+{
+	switch (Value)
+	{
+	case EInvitationResponse::UnknownFailure:
+	{
+		return TEXT("UnknownFailure");
+	}
+	case EInvitationResponse::BadBuild:
+	{
+		return TEXT("BadBuild");
+	}
+	case EInvitationResponse::Rejected:
+	{
+		return TEXT("Rejected");
+	}
+	case EInvitationResponse::Accepted:
+	{
+		return TEXT("Accepted");
+	}
+	}
+	return TEXT("Unknown");
+}
+
+const TCHAR* ToString(const PartySystemPermissions::EPermissionType Value)
+{
+	switch (Value)
+	{
+	case PartySystemPermissions::EPermissionType::Noone:
+	{
+		return TEXT("Noone");
+	}
+	case PartySystemPermissions::EPermissionType::Leader:
+	{
+		return TEXT("Leader");
+	}
+	case PartySystemPermissions::EPermissionType::Friends:
+	{
+		return TEXT("Friends");
+	}
+	case PartySystemPermissions::EPermissionType::Anyone:
+	{
+		return TEXT("Anyone");
+	}
+	}
+	return TEXT("Unknown");
+}
+
+PartySystemPermissions::EPermissionType PartySystemPermissionTypeFromString(const TCHAR* Value)
+{
+	if (FCString::Stricmp(Value, TEXT("Leader")) == 0)
+	{
+		return PartySystemPermissions::EPermissionType::Leader;
+	}
+	else if (FCString::Stricmp(Value, TEXT("Friends")) == 0)
+	{
+		return PartySystemPermissions::EPermissionType::Friends;
+	}
+	else if (FCString::Stricmp(Value, TEXT("Anyone")) == 0)
+	{
+		return PartySystemPermissions::EPermissionType::Anyone;
+	}
+	return PartySystemPermissions::EPermissionType::Noone;
+}
+
+const TCHAR* ToString(const EJoinRequestAction Value)
+{
+	switch (Value)
+	{
+	case EJoinRequestAction::Manual:
+	{
+		return TEXT("Manual");
+	}
+	case EJoinRequestAction::AutoApprove:
+	{
+		return TEXT("AutoApprove");
+	}
+	case EJoinRequestAction::AutoReject:
+	{
+		return TEXT("AutoReject");
+	}
+	}
+	return TEXT("Unknown");
+}
+
+EJoinRequestAction JoinRequestActionFromString(const TCHAR* Value)
+{
+	if (FCString::Stricmp(Value, TEXT("AutoApprove")) == 0)
+	{
+		return EJoinRequestAction::AutoApprove;
+	}
+	else if (FCString::Stricmp(Value, TEXT("AutoReject")) == 0)
+	{
+		return EJoinRequestAction::AutoReject;
+	}
+	return EJoinRequestAction::Manual;
+}
+
+FString ToDebugString(const FPartyConfiguration& PartyConfiguration)
+{
+	return FString::Printf(TEXT("JoinRequestAction(%s) RemoveOnDisconnect(%d) Publish(%s) Chat(%d) Invite(%s) Accepting(%d) Not Accepting Reason(%d) MaxMembers: %d Nickname: %s Description: %s Password: %s"),
+			ToString(PartyConfiguration.JoinRequestAction),
+			PartyConfiguration.bShouldRemoveOnDisconnection,
+			ToString(PartyConfiguration.PresencePermissions),
+			PartyConfiguration.bChatEnabled,
+			ToString(PartyConfiguration.InvitePermissions),
+			PartyConfiguration.bIsAcceptingMembers,
+			PartyConfiguration.NotAcceptingMembersReason,
+			PartyConfiguration.MaxMembers,
+			*PartyConfiguration.Nickname,
+			*PartyConfiguration.Description,
+			PartyConfiguration.Password.IsEmpty() ? TEXT("not set") : *PartyConfiguration.Password
+		);
+}
+
+FString ToDebugString(const IOnlinePartyJoinInfo& JoinInfo)
+{
+	return FString::Printf(TEXT("SourceUserId(%s) SourceDisplayName(%s) PartyId(%s) HasKey(%d) HasPassword(%d) IsAcceptingMembers(%d) NotAcceptingReason(%d)"),
+			*(JoinInfo.GetSourceUserId()->ToDebugString()),
+			*(JoinInfo.GetSourceDisplayName()),
+			*(JoinInfo.GetPartyId()->ToDebugString()),
+			JoinInfo.HasKey() ? 1 : 0,
+			JoinInfo.HasPassword() ? 1 : 0,
+			JoinInfo.IsAcceptingMembers() ? 1 : 0,
+			JoinInfo.GetNotAcceptingReason()
+		);
+}
+
+/**
+ * Dump key/value pairs for debugging
+ */
+FString ToDebugString(const FOnlineKeyValuePairs<FString, FVariantData>& KeyValAttrs)
+{
+	FString Result;
+	bool bPrintedAny = false;
+	for (const TPair<FString, FVariantData>& Iterator : KeyValAttrs)
+	{
+		if (bPrintedAny)
+		{
+			Result += TEXT(",");
+		}
+		Result += FString::Printf(TEXT("[%s=%s]"), *Iterator.Key, *Iterator.Value.ToString());
+		bPrintedAny = true;
+	}
+	return Result;
+}
+
+/**
+* Dump state about the party data for debugging
+*/
+FString ToDebugString(const FOnlinePartyData& PartyData)
+{
+	FString Result;
+	
+	int32 TotalBytesPerSec = PartyData.TotalPackets ? (PartyData.TotalBytes / PartyData.TotalPackets) : 0;
+	int32 TotalEffectiveBytesPerSec = PartyData.TotalPackets ? (PartyData.TotalEffectiveBytes / PartyData.TotalPackets) : 0;
+
+	Result += FString::Printf(TEXT("%dB [%d B/pkt], %dB [%d B/pkt], Rev: %d,"), 
+		PartyData.TotalBytes, TotalBytesPerSec,
+		PartyData.TotalEffectiveBytes, TotalEffectiveBytesPerSec, 
+		PartyData.RevisionCount);
+
+	Result += ToDebugString(PartyData.GetKeyValAttrs());
+	return Result;
+}
diff --git a/Engine/Plugins/Online/OnlineSubsystem/Source/Public/Interfaces/OnlinePartyInterface.h b/Engine/Plugins/Online/OnlineSubsystem/Source/Public/Interfaces/OnlinePartyInterface.h
--- a/Engine/Plugins/Online/OnlineSubsystem/Source/Public/Interfaces/OnlinePartyInterface.h
+++ b/Engine/Plugins/Online/OnlineSubsystem/Source/Public/Interfaces/OnlinePartyInterface.h
@@ -1755,562 +1755,44 @@ enum class EInvitationResponse
 };
 
 /** @return the stringified version of the enum passed in */
-inline const TCHAR* ToString(const EPartyState Value)
-{
-	switch (Value)
-	{
-	case EPartyState::None:
-	{
-		return TEXT("None");
-	}
-	case EPartyState::CreatePending:
-	{
-		return TEXT("CreatePending");
-	}
-	case EPartyState::JoinPending:
-	{
-		return TEXT("JoinPending");
-	}
-	case EPartyState::LeavePending:
-	{
-		return TEXT("LeavePending");
-	}
-	case EPartyState::Active:
-	{
-		return TEXT("Active");
-	}
-	case EPartyState::Disconnected:
-	{
-		return TEXT("Disconnected");
-	}
-	case EPartyState::CleanUp:
-	{
-		return TEXT("CleanUp");
-	}
-	}
-	return TEXT("Unknown");
-}
-
-inline const TCHAR* ToString(const EMemberExitedReason Value)
-{
-	switch (Value)
-	{
-	case EMemberExitedReason::Unknown:
-	{
-		return TEXT("Unknown");
-	}
-	case EMemberExitedReason::Left:
-	{
-		return TEXT("Left");
-	}
-	case EMemberExitedReason::Removed:
-	{
-		return TEXT("Removed");
-	}
-	case EMemberExitedReason::Kicked:
-	{
-		return TEXT("Kicked");
-	}
-	}
-	return TEXT("Unknown"); // Same as EMemberExitedReason::Unknown, which is ok because it is only used when we do not have enough information
-}
-
-inline const TCHAR* ToString(const ECreatePartyCompletionResult Value)
-{
-	switch (Value)
-	{
-	case ECreatePartyCompletionResult::UnknownClientFailure:
-	{
-		return TEXT("UnknownClientFailure");
-	}
-	case ECreatePartyCompletionResult::AlreadyCreatingParty:
-	{
-		return TEXT("AlreadyCreatingParty");
-	}
-	case ECreatePartyCompletionResult::AlreadyInParty:
-	{
-		return TEXT("AlreadyInParty");
-	}
-	case ECreatePartyCompletionResult::FailedToCreateMucRoom:
-	{
-		return TEXT("FailedToCreateMucRoom");
-	}
-	case ECreatePartyCompletionResult::NoResponse:
-	{
-		return TEXT("NoResponse");
-	}
-	case ECreatePartyCompletionResult::LoggedOut:
-	{
-		return TEXT("LoggedOut");
-	}
-	case ECreatePartyCompletionResult::UnknownInternalFailure:
-	{
-		return TEXT("UnknownInternalFailure");
-	}
-	case ECreatePartyCompletionResult::Succeeded:
-	{
-		return TEXT("Succeeded");
-	}
-	}
-	return TEXT("Unknown");
-}
-
-inline const TCHAR* ToString(const ESendPartyInvitationCompletionResult Value)
-{
-	switch (Value)
-	{
-	case ESendPartyInvitationCompletionResult::NotLoggedIn:
-	{
-		return TEXT("NotLoggedIn");
-	}
-	case ESendPartyInvitationCompletionResult::InvitePending:
-	{
-		return TEXT("InvitePending");
-	}
-	case ESendPartyInvitationCompletionResult::AlreadyInParty:
-	{
-		return TEXT("AlreadyInParty");
-	}
-	case ESendPartyInvitationCompletionResult::PartyFull:
-	{
-		return TEXT("PartyFull");
-	}
-	case ESendPartyInvitationCompletionResult::NoPermission:
-	{
-		return TEXT("NoPermission");
-	}
-	case ESendPartyInvitationCompletionResult::UnknownInternalFailure:
-	{
-		return TEXT("UnknownInternalFailure");
-	}
-	case ESendPartyInvitationCompletionResult::Succeeded:
-	{
-		return TEXT("Succeeded");
-	}
-	}
-	return TEXT("Unknown");
-}
-
-inline const TCHAR* ToString(const EJoinPartyCompletionResult Value)
-{
-	switch (Value)
-	{
-	case EJoinPartyCompletionResult::UnknownClientFailure:
-	{
-		return TEXT("UnknownClientFailure");
-	}
-	case EJoinPartyCompletionResult::BadBuild:
-	{
-		return TEXT("BadBuild");
-	}
-	case EJoinPartyCompletionResult::InvalidAccessKey:
-	{
-		return TEXT("InvalidAccessKey");
-	}
-	case EJoinPartyCompletionResult::AlreadyInLeadersJoiningList:
-	{
-		return TEXT("AlreadyInLeadersJoiningList");
-	}
-	case EJoinPartyCompletionResult::AlreadyInLeadersPartyRoster:
-	{
-		return TEXT("AlreadyInLeadersPartyRoster");
-	}
-	case EJoinPartyCompletionResult::NoSpace:
-	{
-		return TEXT("NoSpace");
-	}
-	case EJoinPartyCompletionResult::NotApproved:
-	{
-		return TEXT("NotApproved");
-	}
-	case EJoinPartyCompletionResult::RequesteeNotMember:
-	{
-		return TEXT("RequesteeNotMember");
-	}
-	case EJoinPartyCompletionResult::RequesteeNotLeader:
-	{
-		return TEXT("RequesteeNotLeader");
-	}
-	case EJoinPartyCompletionResult::NoResponse:
-	{
-		return TEXT("NoResponse");
-	}
-	case EJoinPartyCompletionResult::LoggedOut:
-	{
-		return TEXT("LoggedOut");
-	}
-	case EJoinPartyCompletionResult::UnableToRejoin:
-	{
-		return TEXT("UnableToRejoin");
-	}
-	case EJoinPartyCompletionResult::IncompatiblePlatform:
-	{
-		return TEXT("IncompatiblePlatform");
-	}
-	case EJoinPartyCompletionResult::AlreadyJoiningParty:
-	{
-		return TEXT("AlreadyJoiningParty");
-	}
-	case EJoinPartyCompletionResult::AlreadyInParty:
-	{
-		return TEXT("AlreadyInParty");
-	}
-	case EJoinPartyCompletionResult::JoinInfoInvalid:
-	{
-		return TEXT("JoinInfoInvalid");
-	}
-	case EJoinPartyCompletionResult::AlreadyInPartyOfSpecifiedType:
-	{
-		return TEXT("AlreadyInPartyOfSpecifiedType");
-	}
-	case EJoinPartyCompletionResult::MessagingFailure:
-	{
-		return TEXT("MessagingFailure");
-	}
-	case EJoinPartyCompletionResult::GameSpecificReason:
-	{
-		return TEXT("GameSpecificReason");
-	}
-	case EJoinPartyCompletionResult::Succeeded:
-	{
-		return TEXT("Succeeded");
-	}
-	case EJoinPartyCompletionResult::UnknownInternalFailure:
-	{
-		return TEXT("DeprecatedUnknownInternalFailure");
-	}
-	}
-	return TEXT("Unknown");
-}
-
-inline const TCHAR* ToString(const ELeavePartyCompletionResult Value)
-{
-	switch (Value)
-	{
-	case ELeavePartyCompletionResult::UnknownClientFailure:
-	{
-		return TEXT("UnknownClientFailure");
-	}
-	case ELeavePartyCompletionResult::NoResponse:
-	{
-		return TEXT("NoResponse");
-	}
-	case ELeavePartyCompletionResult::LoggedOut:
-	{
-		return TEXT("LoggedOut");
-	}
-	case ELeavePartyCompletionResult::UnknownParty:
-	{
-		return TEXT("UnknownParty");
-	}
-	case ELeavePartyCompletionResult::LeavePending:
-	{
-		return TEXT("LeavePending");
-	}
-	case ELeavePartyCompletionResult::Succeeded:
-	{
-		return TEXT("Succeeded");
-	}
-	case ELeavePartyCompletionResult::UnknownLocalUser:
-	{
-		return TEXT("DeprecatedUnknownLocalUser");
-	}
-	case ELeavePartyCompletionResult::NotMember:
-	{
-		return TEXT("DeprecatedNotMember");
-	}
-	case ELeavePartyCompletionResult::MessagingFailure:
-	{
-		return TEXT("DeprecatedMessagingFailure");
-	}
-	case ELeavePartyCompletionResult::UnknownTransportFailure:
-	{
-		return TEXT("DeprecatedUnknownTransportFailure");
-	}
-	case ELeavePartyCompletionResult::UnknownInternalFailure:
-	{
-		return TEXT("DeprecatedUnknownInternalFailure");
-	}
-	}
-	return TEXT("Unknown");
-}
-
-inline const TCHAR* ToString(const EUpdateConfigCompletionResult Value)
-{
-	switch (Value)
-	{
-	case EUpdateConfigCompletionResult::UnknownClientFailure:
-	{
-		return TEXT("UnknownClientFailure");
-	}
-	case EUpdateConfigCompletionResult::UnknownParty:
-	{
-		return TEXT("UnknownParty");
-	}
-	case EUpdateConfigCompletionResult::LocalMemberNotMember:
-	{
-		return TEXT("LocalMemberNotMember");
-	}
-	case EUpdateConfigCompletionResult::LocalMemberNotLeader:
-	{
-		return TEXT("LocalMemberNotLeader");
-	}
-	case EUpdateConfigCompletionResult::RemoteMemberNotMember:
-	{
-		return TEXT("RemoteMemberNotMember");
-	}
-	case EUpdateConfigCompletionResult::MessagingFailure:
-	{
-		return TEXT("MessagingFailure");
-	}
-	case EUpdateConfigCompletionResult::NoResponse:
-	{
-		return TEXT("NoResponse");
-	}
-	case EUpdateConfigCompletionResult::UnknownInternalFailure:
-	{
-		return TEXT("UnknownInternalFailure");
-	}
-	case EUpdateConfigCompletionResult::Succeeded:
-	{
-		return TEXT("Succeeded");
-	}
-	}
-	return TEXT("Unknown");
-}
-
-inline const TCHAR* ToString(const EKickMemberCompletionResult Value)
-{
-	switch (Value)
-	{
-	case EKickMemberCompletionResult::UnknownClientFailure:
-	{
-		return TEXT("UnknownClientFailure");
-	}
-	case EKickMemberCompletionResult::UnknownParty:
-	{
-		return TEXT("UnknownParty");
-	}
-	case EKickMemberCompletionResult::LocalMemberNotMember:
-	{
-		return TEXT("LocalMemberNotMember");
-	}
-	case EKickMemberCompletionResult::LocalMemberNotLeader:
-	{
-		return TEXT("LocalMemberNotLeader");
-	}
-	case EKickMemberCompletionResult::RemoteMemberNotMember:
-	{
-		return TEXT("RemoteMemberNotMember");
-	}
-	case EKickMemberCompletionResult::MessagingFailure:
-	{
-		return TEXT("MessagingFailure");
-	}
-	case EKickMemberCompletionResult::NoResponse:
-	{
-		return TEXT("NoResponse");
-	}
-	case EKickMemberCompletionResult::LoggedOut:
-	{
-		return TEXT("LoggedOut");
-	}
-	case EKickMemberCompletionResult::UnknownInternalFailure:
-	{
-		return TEXT("UnknownInternalFailure");
-	}
-	case EKickMemberCompletionResult::Succeeded:
-	{
-		return TEXT("Succeeded");
-	}
-	}
-	return TEXT("Unknown");
-}
-
-inline const TCHAR* ToString(const EPromoteMemberCompletionResult Value)
-{
-	switch (Value)
-	{
-	case EPromoteMemberCompletionResult::UnknownClientFailure:
-	{
-		return TEXT("UnknownClientFailure");
-	}
-	case EPromoteMemberCompletionResult::UnknownParty:
-	{
-		return TEXT("UnknownParty");
-	}
-	case EPromoteMemberCompletionResult::LocalMemberNotMember:
-	{
-		return TEXT("LocalMemberNotMember");
-	}
-	case EPromoteMemberCompletionResult::LocalMemberNotLeader:
-	{
-		return TEXT("LocalMemberNotLeader");
-	}
-	case EPromoteMemberCompletionResult::TargetIsSelf:
-	{
-		return TEXT("TargetIsSelf");
-	}
-	case EPromoteMemberCompletionResult::TargetNotMember:
-	{
-		return TEXT("TargetNotMember");
-	}
-	case EPromoteMemberCompletionResult::MessagingFailure:
-	{
-		return TEXT("MessagingFailure");
-	}
-	case EPromoteMemberCompletionResult::NoResponse:
-	{
-		return TEXT("NoResponse");
-	}
-	case EPromoteMemberCompletionResult::LoggedOut:
-	{
-		return TEXT("LoggedOut");
-	}
-	case EPromoteMemberCompletionResult::UnknownInternalFailure:
-	{
-		return TEXT("UnknownInternalFailure");
-	}
-	case EPromoteMemberCompletionResult::Succeeded:
-	{
-		return TEXT("Succeeded");
-	}
-	}
-	return TEXT("Unknown");
-}
-
-inline const TCHAR* ToString(const PartySystemPermissions::EPermissionType Value)
-{
-	switch (Value)
-	{
-	case PartySystemPermissions::EPermissionType::Noone:
-	{
-		return TEXT("Noone");
-	}
-	case PartySystemPermissions::EPermissionType::Leader:
-	{
-		return TEXT("Leader");
-	}
-	case PartySystemPermissions::EPermissionType::Friends:
-	{
-		return TEXT("Friends");
-	}
-	case PartySystemPermissions::EPermissionType::Anyone:
-	{
-		return TEXT("Anyone");
-	}
-	}
-	return TEXT("Unknown");
-}
-
-inline const TCHAR* ToString(const EJoinRequestAction Value)
-{
-	switch (Value)
-	{
-	case EJoinRequestAction::Manual:
-	{
-		return TEXT("Manual");
-	}
-	case EJoinRequestAction::AutoApprove:
-	{
-		return TEXT("AutoApprove");
-	}
-	case EJoinRequestAction::AutoReject:
-	{
-		return TEXT("AutoReject");
-	}
-	}
-	return TEXT("Unknown");
-}
-
-inline const TCHAR* ToString(const EInvitationResponse Value)
-{
-	switch (Value)
-	{
-	case EInvitationResponse::UnknownFailure:
-	{
-		return TEXT("UnknownFailure");
-	}
-	case EInvitationResponse::BadBuild:
-	{
-		return TEXT("BadBuild");
-	}
-	case EInvitationResponse::Rejected:
-	{
-		return TEXT("Rejected");
-	}
-	case EInvitationResponse::Accepted:
-	{
-		return TEXT("Accepted");
-	}
-	}
-	return TEXT("Unknown");
-}
-
-inline FString ToDebugString(const FPartyConfiguration& PartyConfiguration)
-{
-	return FString::Printf(TEXT("JoinRequestAction(%s) RemoveOnDisconnect(%d) Publish(%s) Chat(%d) Invite(%s) Accepting(%d) Not Accepting Reason(%d) MaxMembers: %d Nickname: %s Description: %s Password: %s"),
-			ToString(PartyConfiguration.JoinRequestAction),
-			PartyConfiguration.bShouldRemoveOnDisconnection,
-			ToString(PartyConfiguration.PresencePermissions),
-			PartyConfiguration.bChatEnabled,
-			ToString(PartyConfiguration.InvitePermissions),
-			PartyConfiguration.bIsAcceptingMembers,
-			PartyConfiguration.NotAcceptingMembersReason,
-			PartyConfiguration.MaxMembers,
-			*PartyConfiguration.Nickname,
-			*PartyConfiguration.Description,
-			PartyConfiguration.Password.IsEmpty() ? TEXT("not set") : *PartyConfiguration.Password
-		);
-}
-
-inline FString ToDebugString(const IOnlinePartyJoinInfo& JoinInfo)
-{
-	return FString::Printf(TEXT("SourceUserId(%s) SourceDisplayName(%s) PartyId(%s) HasKey(%d) HasPassword(%d) IsAcceptingMembers(%d) NotAcceptingReason(%d)"),
-			*(JoinInfo.GetSourceUserId()->ToDebugString()),
-			*(JoinInfo.GetSourceDisplayName()),
-			*(JoinInfo.GetPartyId()->ToDebugString()),
-			JoinInfo.HasKey() ? 1 : 0,
-			JoinInfo.HasPassword() ? 1 : 0,
-			JoinInfo.IsAcceptingMembers() ? 1 : 0,
-			JoinInfo.GetNotAcceptingReason()
-		);
-}
-
-/**
- * Dump key/value pairs for debugging
- */
-inline FString ToDebugString(const FOnlineKeyValuePairs<FString, FVariantData>& KeyValAttrs)
-{
-	FString Result;
-	bool bPrintedAny = false;
-	for (const TPair<FString, FVariantData>& Iterator : KeyValAttrs)
-	{
-		if (bPrintedAny)
-		{
-			Result += TEXT(",");
-		}
-		Result += FString::Printf(TEXT("[%s=%s]"), *Iterator.Key, *Iterator.Value.ToString());
-		bPrintedAny = true;
-	}
-	return Result;
-}
+ONLINESUBSYSTEM_API const TCHAR* ToString(const EPartyState Value);
+/** @return the enum version of the string passed in */
+ONLINESUBSYSTEM_API EPartyState EPartyStateFromString(const TCHAR* Value);
 
-/**
-* Dump state about the party data for debugging
-*/
-inline FString ToDebugString(const FOnlinePartyData& PartyData)
-{
-	FString Result;
-	
-	int32 TotalBytesPerSec = PartyData.TotalPackets ? (PartyData.TotalBytes / PartyData.TotalPackets) : 0;
-	int32 TotalEffectiveBytesPerSec = PartyData.TotalPackets ? (PartyData.TotalEffectiveBytes / PartyData.TotalPackets) : 0;
+/** @return the stringified version of the enum passed in */
+ONLINESUBSYSTEM_API const TCHAR* ToString(const EMemberExitedReason Value);
+/** @return the stringified version of the enum passed in */
+ONLINESUBSYSTEM_API const TCHAR* ToString(const ECreatePartyCompletionResult Value);
+/** @return the stringified version of the enum passed in */
+ONLINESUBSYSTEM_API const TCHAR* ToString(const ESendPartyInvitationCompletionResult Value);
+/** @return the stringified version of the enum passed in */
+ONLINESUBSYSTEM_API const TCHAR* ToString(const EJoinPartyCompletionResult Value);
+/** @return the stringified version of the enum passed in */
+ONLINESUBSYSTEM_API const TCHAR* ToString(const ELeavePartyCompletionResult Value);
+/** @return the stringified version of the enum passed in */
+ONLINESUBSYSTEM_API const TCHAR* ToString(const EUpdateConfigCompletionResult Value);
+/** @return the stringified version of the enum passed in */
+ONLINESUBSYSTEM_API const TCHAR* ToString(const EKickMemberCompletionResult Value);
+/** @return the stringified version of the enum passed in */
+ONLINESUBSYSTEM_API const TCHAR* ToString(const EPromoteMemberCompletionResult Value);
+/** @return the stringified version of the enum passed in */
+ONLINESUBSYSTEM_API const TCHAR* ToString(const EInvitationResponse Value);
 
-	Result += FString::Printf(TEXT("%dB [%d B/pkt], %dB [%d B/pkt], Rev: %d,"), 
-		PartyData.TotalBytes, TotalBytesPerSec,
-		PartyData.TotalEffectiveBytes, TotalEffectiveBytesPerSec, 
-		PartyData.RevisionCount);
+/** @return the stringified version of the enum passed in */
+ONLINESUBSYSTEM_API const TCHAR* ToString(const PartySystemPermissions::EPermissionType Value);
+/** @return the enum version of the string passed in */
+ONLINESUBSYSTEM_API PartySystemPermissions::EPermissionType PartySystemPermissionTypeFromString(const TCHAR* Value);
 
-	Result += ToDebugString(PartyData.GetKeyValAttrs());
-	return Result;
-}
+/** @return the stringified version of the enum passed in */
+ONLINESUBSYSTEM_API const TCHAR* ToString(const EJoinRequestAction Value);
+/** @return the enum version of the string passed in */
+ONLINESUBSYSTEM_API EJoinRequestAction JoinRequestActionFromString(const TCHAR* Value);
+
+/** Dump party configuration for debugging */
+ONLINESUBSYSTEM_API FString ToDebugString(const FPartyConfiguration& PartyConfiguration);
+/** Dump join info for debugging */
+ONLINESUBSYSTEM_API FString ToDebugString(const IOnlinePartyJoinInfo& JoinInfo);
+/** Dump key/value pairs for debugging */
+ONLINESUBSYSTEM_API FString ToDebugString(const FOnlineKeyValuePairs<FString, FVariantData>& KeyValAttrs);
+/** Dump state about the party data for debugging */
+ONLINESUBSYSTEM_API FString ToDebugString(const FOnlinePartyData& PartyData);
\ No newline at end of file
