Use shared pointers instead of raw pointers (to shared pointers) in GetOnlineSubsystem.  This should fix any issues with FOnlineSubsystemDelegates::OnOnlineSubsystemCreated having side effects

diff --git a/Engine/Plugins/Online/OnlineSubsystem/Source/Private/OnlineSubsystemModule.cpp b/Engine/Plugins/Online/OnlineSubsystem/Source/Private/OnlineSubsystemModule.cpp
--- a/Engine/Plugins/Online/OnlineSubsystem/Source/Private/OnlineSubsystemModule.cpp
+++ b/Engine/Plugins/Online/OnlineSubsystem/Source/Private/OnlineSubsystemModule.cpp
@@ -130,15 +130,35 @@ void FOnlineSubsystemModule::ProcessConfigDefinedSubsystems()
 bool FOnlineSubsystemModule::TryLoadSubsystemAndSetDefault(FName ModuleName)
 {
 	// A module loaded with its factory method set for creation and a default instance of the online subsystem is required
-	if ((LoadSubsystemModule(ModuleName.ToString()) &&
-		OnlineFactories.Contains(ModuleName) &&
-		GetOnlineSubsystem(ModuleName) != nullptr))
+	bool bLoaded = false;
+	FString ModuleNameString = ModuleName.ToString();
+	if (LoadSubsystemModule(ModuleNameString))
 	{
-		DefaultPlatformService = ModuleName;
-		return true;
+		if (OnlineFactories.Contains(ModuleName))
+		{
+			IOnlineSubsystem* OnlineSubsystem = GetOnlineSubsystem(ModuleName);
+			if (OnlineSubsystem != nullptr)
+			{
+				UE_LOG_ONLINE(Log, TEXT("TryLoadSubsystemAndSetDefault: Loaded subsystem for module [%s]"), *ModuleNameString);
+				DefaultPlatformService = ModuleName;
+				bLoaded = true;
+			}
+			else
+			{
+				UE_LOG_ONLINE(Warning, TEXT("TryLoadSubsystemAndSetDefault: GetOnlineSubsystem([%s]) failed"), *ModuleNameString);
+			}
+		}
+		else
+		{
+			UE_LOG_ONLINE(Warning, TEXT("TryLoadSubsystemAndSetDefault: OnlineFactories does not contain [%s]"), *ModuleNameString);
+		}
+	}
+	else
+	{
+		UE_LOG_ONLINE(Warning, TEXT("TryLoadSubsystemAndSetDefault: LoadSubsystemModule([%s]) failed"), *ModuleNameString);
 	}
 
-	return false;
+	return bLoaded;
 }
 
 void FOnlineSubsystemModule::LoadDefaultSubsystem()
@@ -155,6 +175,7 @@ void FOnlineSubsystemModule::LoadDefaultSubsystem()
 	// if default fails, attempt to load Null
 	if (!bHasLoadedModule)
 	{
+		UE_LOG_ONLINE(Warning, TEXT("LoadDefaultSubsystem: Failed to load Subsystem=[%s], falling back to NULL_SUBSYSTEM"), *InterfaceString);
 		bHasLoadedModule = TryLoadSubsystemAndSetDefault(NULL_SUBSYSTEM);
 	}
 
@@ -283,11 +304,15 @@ IOnlineSubsystem* FOnlineSubsystemModule::GetOnlineSubsystem(const FName InSubsy
 	FName SubsystemName, InstanceName;
 	FName KeyName = ParseOnlineSubsystemName(InSubsystemName, SubsystemName, InstanceName);
 
-	IOnlineSubsystemPtr* OnlineSubsystem = nullptr;
+	IOnlineSubsystemPtr OnlineSubsystem;
 	if (!SubsystemName.IsNone())
 	{
-		OnlineSubsystem = OnlineSubsystems.Find(KeyName);
-		if (OnlineSubsystem == nullptr)
+		IOnlineSubsystemPtr* ExistingOnlineSubsystem = OnlineSubsystems.Find(KeyName);
+		if (ExistingOnlineSubsystem != nullptr)
+		{
+			OnlineSubsystem = *ExistingOnlineSubsystem;
+		}
+		else
 		{
 			if (IOnlineSubsystem::IsEnabled(SubsystemName))
 			{
@@ -311,7 +336,7 @@ IOnlineSubsystem* FOnlineSubsystemModule::GetOnlineSubsystem(const FName InSubsy
 					if (NewSubsystemInstance.IsValid())
 					{
 						OnlineSubsystems.Add(KeyName, NewSubsystemInstance);
-						OnlineSubsystem = OnlineSubsystems.Find(KeyName);
+						OnlineSubsystem = NewSubsystemInstance;
 						FOnlineSubsystemDelegates::OnOnlineSubsystemCreated.Broadcast(NewSubsystemInstance.Get());						
 					}
 					else
@@ -328,7 +353,7 @@ IOnlineSubsystem* FOnlineSubsystemModule::GetOnlineSubsystem(const FName InSubsy
 		}
 	}
 
-	return (OnlineSubsystem == nullptr) ? nullptr : (*OnlineSubsystem).Get();
+	return OnlineSubsystem.Get();
 }
 
 IOnlineSubsystem* FOnlineSubsystemModule::GetNativeSubsystem(bool bAutoLoad)
