Fixup some API docs

diff --git a/Engine/Plugins/Online/OnlineSubsystem/Source/Public/Interfaces/OnlineIdentityInterface.h b/Engine/Plugins/Online/OnlineSubsystem/Source/Public/Interfaces/OnlineIdentityInterface.h
--- a/Engine/Plugins/Online/OnlineSubsystem/Source/Public/Interfaces/OnlineIdentityInterface.h
+++ b/Engine/Plugins/Online/OnlineSubsystem/Source/Public/Interfaces/OnlineIdentityInterface.h
@@ -149,7 +149,7 @@ typedef FOnLoginFlowLogout::FDelegate FOnLoginFlowLogoutDelegate;
  * @param UserId The unique id of the user who was queried
  * @param OnlineError the result of the operation
  */
-DECLARE_DELEGATE_TwoParams(FOnRevokeAuthTokenCompleteDelegate, const FUniqueNetId&, const FOnlineError&);
+DECLARE_DELEGATE_TwoParams(FOnRevokeAuthTokenCompleteDelegate, const FUniqueNetId& /*UserId*/, const FOnlineError& /*OnlineError*/);
 
 /**
  * Interface for registration/authentication of user identities
@@ -197,7 +197,7 @@ public:
 	 *
 	 * @param LocalUserNum the player that logged in/out
 	 */
-	DEFINE_ONLINE_DELEGATE_ONE_PARAM(OnLoginChanged, int32);
+	DEFINE_ONLINE_DELEGATE_ONE_PARAM(OnLoginChanged, int32 /*LocalUserNum*/);
 
 	/**
 	 * Delegate called when a player's login status changes but doesn't change identity
@@ -207,7 +207,7 @@ public:
 	 * @param NewStatus the new login status for the user
 	 * @param NewId the new id to associate with the user
 	 */
-	DEFINE_ONLINE_PLAYER_DELEGATE_THREE_PARAM(MAX_LOCAL_PLAYERS, OnLoginStatusChanged, ELoginStatus::Type /*Previous*/, ELoginStatus::Type /*Current*/, const FUniqueNetId&);
+	DEFINE_ONLINE_PLAYER_DELEGATE_THREE_PARAM(MAX_LOCAL_PLAYERS, OnLoginStatusChanged, ELoginStatus::Type /*OldStatus*/, ELoginStatus::Type /*NewStatus*/, const FUniqueNetId& /*NewId*/);
 
 	/**
 	 * Delegate called when a controller-user pairing changes
@@ -216,7 +216,7 @@ public:
 	 * @param PreviousUser the new login status for the user
 	 * @param NewUser the new id to associate with the user
 	 */
-	DEFINE_ONLINE_DELEGATE_THREE_PARAM(OnControllerPairingChanged, int, const FUniqueNetId&, const FUniqueNetId&);
+	DEFINE_ONLINE_DELEGATE_THREE_PARAM(OnControllerPairingChanged, int /*LocalUserNum*/, const FUniqueNetId& /*PreviousUser*/, const FUniqueNetId& /*NewUser*/);
 
 	/**
 	 * Login/Authenticate with user credentials.
@@ -232,12 +232,11 @@ public:
 	/**
 	 * Called when user account login has completed after calling Login() or AutoLogin()
 	 *
-	 * @param LocalUserNum the controller number of the associated user
 	 * @param bWasSuccessful true if server was contacted and a valid result received
 	 * @param UserId the user id received from the server on successful login
 	 * @param Error string representing the error condition
 	 */
-	DEFINE_ONLINE_PLAYER_DELEGATE_THREE_PARAM(MAX_LOCAL_PLAYERS, OnLoginComplete, bool, const FUniqueNetId&, const FString&);
+	DEFINE_ONLINE_PLAYER_DELEGATE_THREE_PARAM(MAX_LOCAL_PLAYERS, OnLoginComplete, bool /*bWasSuccessful*/, const FUniqueNetId& /*UserId*/, const FString& /*Error*/);
 
 	/**
 	 * Signs the player out of the online service
@@ -252,10 +251,9 @@ public:
 	/**
 	 * Delegate used in notifying the that manual logout completed
 	 *
-	 * @param LocalUserNum the controller number of the associated user
 	 * @param bWasSuccessful whether the async call completed properly or not
 	 */
-	DEFINE_ONLINE_PLAYER_DELEGATE_ONE_PARAM(MAX_LOCAL_PLAYERS, OnLogoutComplete, bool);
+	DEFINE_ONLINE_PLAYER_DELEGATE_ONE_PARAM(MAX_LOCAL_PLAYERS, OnLogoutComplete, bool /*bWasSuccessful*/);
 
 	/**
 	 * Delegate called when the online subsystem requires the login flow to logout and cleanup
@@ -387,11 +385,11 @@ public:
 	/**
 	 * Delegate executed when we get a user privilege result.
 	 *
-	 * @param UniqueId The unique id of the user who was queried
+	 * @param LocalUserId The unique id of the user who was queried
 	 * @param Privilege the privilege that was queried
 	 * @param PrivilegeResult bitwise OR of any privilege failures. 0 is success.
 	 */
-	DECLARE_DELEGATE_ThreeParams(FOnGetUserPrivilegeCompleteDelegate, const FUniqueNetId&, EUserPrivileges::Type, uint32);
+	DECLARE_DELEGATE_ThreeParams(FOnGetUserPrivilegeCompleteDelegate, const FUniqueNetId& /*LocalUserId*/, EUserPrivileges::Type /*Privilege*/, uint32 /*PrivilegeResult*/);
 
 	/**
 	 * Gets the status of a user's privilege.
