Interface docs fixes

diff --git a/Engine/Plugins/Online/OnlineSubsystem/Source/Public/Interfaces/OnlineIdentityInterface.h b/Engine/Plugins/Online/OnlineSubsystem/Source/Public/Interfaces/OnlineIdentityInterface.h
--- a/Engine/Plugins/Online/OnlineSubsystem/Source/Public/Interfaces/OnlineIdentityInterface.h
+++ b/Engine/Plugins/Online/OnlineSubsystem/Source/Public/Interfaces/OnlineIdentityInterface.h
@@ -212,9 +212,9 @@ public:
 	/**
 	 * Delegate called when a controller-user pairing changes
 	 *
-	 * @param LocalUserNum the controller number of the associated user
-	 * @param PreviousUser the new login status for the user
-	 * @param NewUser the new id to associate with the user
+	 * @param LocalUserNum the controller number whose pairing changed
+	 * @param PreviousUser the previous user associated with the controller
+	 * @param NewUser the new user associated with the controller
 	 */
 	DEFINE_ONLINE_DELEGATE_THREE_PARAM(OnControllerPairingChanged, int /*LocalUserNum*/, const FUniqueNetId& /*PreviousUser*/, const FUniqueNetId& /*NewUser*/);
 
@@ -232,6 +232,7 @@ public:
 	/**
 	 * Called when user account login has completed after calling Login() or AutoLogin()
 	 *
+	 * @param LocalUserNum the controller number of the associated user
 	 * @param bWasSuccessful true if server was contacted and a valid result received
 	 * @param UserId the user id received from the server on successful login
 	 * @param Error string representing the error condition
@@ -251,12 +252,14 @@ public:
 	/**
 	 * Delegate used in notifying the that manual logout completed
 	 *
+	 * @param LocalUserNum the controller number of the associated user
 	 * @param bWasSuccessful whether the async call completed properly or not
 	 */
 	DEFINE_ONLINE_PLAYER_DELEGATE_ONE_PARAM(MAX_LOCAL_PLAYERS, OnLogoutComplete, bool /*bWasSuccessful*/);
 
 	/**
 	 * Delegate called when the online subsystem requires the login flow to logout and cleanup
+	 *
 	 * @param LoginDomains login domains to be cleaned up
 	 */
 	DEFINE_ONLINE_DELEGATE_ONE_PARAM(OnLoginFlowLogout, const TArray<FString>& /*LoginDomains*/);
@@ -285,7 +288,7 @@ public:
 	/**
 	 * Obtain list of all known/registered user accounts
 	 *
-	 * @return info about the user if found, NULL otherwise
+	 * @return info about the users if found, NULL otherwise
 	 */
 	virtual TArray<TSharedPtr<FUserOnlineAccount> > GetAllUserAccounts() const = 0;
 
@@ -404,6 +407,7 @@ public:
 	 * Temporary hack to get a corresponding FUniqueNetId from a PlatformUserId
 	 *
 	 * @param UniqueNetId The unique id to look up
+	 *
 	 * @return The corresponding id or PLATFORMID_NONE if not found
 	 */
 	virtual FPlatformUserId GetPlatformUserIdFromUniqueNetId(const FUniqueNetId& UniqueNetId) const = 0;
