Add SendInvitation result for being rate limited from sending invitations
TODO needs to expose more information about rate limits

diff --git a/Engine/Plugins/Online/OnlineSubsystem/Source/Private/OnlinePartyInterface.cpp b/Engine/Plugins/Online/OnlineSubsystem/Source/Private/OnlinePartyInterface.cpp
--- a/Engine/Plugins/Online/OnlineSubsystem/Source/Private/OnlinePartyInterface.cpp
+++ b/Engine/Plugins/Online/OnlineSubsystem/Source/Private/OnlinePartyInterface.cpp
@@ -417,6 +417,10 @@ const TCHAR* ToString(const ESendPartyInvitationCompletionResult Value)
 	{
 		return TEXT("NoPermission");
 	}
+	case ESendPartyInvitationCompletionResult::RateLimited:
+	{
+		return TEXT("RateLimited");
+	}
 	case ESendPartyInvitationCompletionResult::UnknownInternalFailure:
 	{
 		return TEXT("UnknownInternalFailure");
diff --git a/Engine/Plugins/Online/OnlineSubsystem/Source/Public/Interfaces/OnlinePartyInterface.h b/Engine/Plugins/Online/OnlineSubsystem/Source/Public/Interfaces/OnlinePartyInterface.h
--- a/Engine/Plugins/Online/OnlineSubsystem/Source/Public/Interfaces/OnlinePartyInterface.h
+++ b/Engine/Plugins/Online/OnlineSubsystem/Source/Public/Interfaces/OnlinePartyInterface.h
@@ -1884,6 +1884,7 @@ enum class ESendPartyInvitationCompletionResult : int8
 	AlreadyInParty,
 	PartyFull,
 	NoPermission,
+	RateLimited,
 	UnknownInternalFailure = 0,
 	Succeeded = 1
 };
