Add default copy/move assignment ops and constructors to FOnlinePartyData.

diff --git a/Engine/Plugins/Online/OnlineSubsystem/Source/Public/Interfaces/OnlinePartyInterface.h b/Engine/Plugins/Online/OnlineSubsystem/Source/Public/Interfaces/OnlinePartyInterface.h
--- a/Engine/Plugins/Online/OnlineSubsystem/Source/Public/Interfaces/OnlinePartyInterface.h
+++ b/Engine/Plugins/Online/OnlineSubsystem/Source/Public/Interfaces/OnlinePartyInterface.h
@@ -108,6 +108,12 @@ public:
 	FOnlinePartyData() = default;
 	virtual ~FOnlinePartyData() = default;
 
+	FOnlinePartyData(const FOnlinePartyData&) = default;
+	FOnlinePartyData& operator=(const FOnlinePartyData&) = default;
+
+	FOnlinePartyData(FOnlinePartyData&&) = default;
+	FOnlinePartyData& operator=(FOnlinePartyData&&) = default;
+
 	/**
 	 * Equality operator
 	 *
