Query for Number of Missed Pings.
    -- SocialToolkit on LocalUser logged in would start query for num missed pings, the value would then be cached when query is successfully complete.

diff --git a/Engine/Plugins/Online/OnlineSubsystem/Source/Public/Interfaces/OnlinePartyInterface.h b/Engine/Plugins/Online/OnlineSubsystem/Source/Public/Interfaces/OnlinePartyInterface.h
--- a/Engine/Plugins/Online/OnlineSubsystem/Source/Public/Interfaces/OnlinePartyInterface.h
+++ b/Engine/Plugins/Online/OnlineSubsystem/Source/Public/Interfaces/OnlinePartyInterface.h
@@ -772,7 +772,7 @@ DECLARE_DELEGATE_ThreeParams(FOnRejectPartyInvitationComplete, const FUniqueNetI
  * @param MemberId - id of member being kicked
  * @param Result - string with error info if any
  */
-DECLARE_DELEGATE_FourParams(FOnKickPartyMemberComplete, const FUniqueNetId& /*LocalUserId*/, const FOnlinePartyId& /*PartyId*/, const FUniqueNetId& /*LocalUserId*/, const EKickMemberCompletionResult /*Result*/);
+DECLARE_DELEGATE_FourParams(FOnKickPartyMemberComplete, const FUniqueNetId& /*LocalUserId*/, const FOnlinePartyId& /*PartyId*/, const FUniqueNetId& /*MemberId*/, const EKickMemberCompletionResult /*Result*/);
 /**
  * Promoting a member of a party async task completed callback
  *
@@ -781,7 +781,7 @@ DECLARE_DELEGATE_FourParams(FOnKickPartyMemberComplete, const FUniqueNetId& /*Lo
  * @param MemberId - id of member being promoted to leader
  * @param Result - string with error info if any
  */
-DECLARE_DELEGATE_FourParams(FOnPromotePartyMemberComplete, const FUniqueNetId& /*LocalUserId*/, const FOnlinePartyId& /*PartyId*/, const FUniqueNetId& /*LocalUserId*/, const EPromoteMemberCompletionResult /*Result*/);
+DECLARE_DELEGATE_FourParams(FOnPromotePartyMemberComplete, const FUniqueNetId& /*LocalUserId*/, const FOnlinePartyId& /*PartyId*/, const FUniqueNetId& /*MemberId*/, const EPromoteMemberCompletionResult /*Result*/);
 
 
 
