Increase MAX_LOCAL_PLAYERS to 4 on windows

diff --git a/Engine/Plugins/Online/OnlineSubsystem/Source/Public/OnlineSubsystemTypes.h b/Engine/Plugins/Online/OnlineSubsystem/Source/Public/OnlineSubsystemTypes.h
--- a/Engine/Plugins/Online/OnlineSubsystem/Source/Public/OnlineSubsystemTypes.h
+++ b/Engine/Plugins/Online/OnlineSubsystem/Source/Public/OnlineSubsystemTypes.h
@@ -30,6 +30,8 @@ extern ONLINESUBSYSTEM_API bool IsUniqueIdLocal(const FUniqueNetId& UniqueId);
 #define MAX_LOCAL_PLAYERS 4
 #elif PLATFORM_SWITCH
 #define MAX_LOCAL_PLAYERS 8
+#elif PLATFORM_WINDOWS
+#define MAX_LOCAL_PLAYERS 4
 #else
 #define MAX_LOCAL_PLAYERS 1
 #endif
