Add RestoreInvites method to IOnlinePartySystem

diff --git a/Engine/Plugins/Online/OnlineSubsystem/Source/Public/Interfaces/OnlinePartyInterface.h b/Engine/Plugins/Online/OnlineSubsystem/Source/Public/Interfaces/OnlinePartyInterface.h
--- a/Engine/Plugins/Online/OnlineSubsystem/Source/Public/Interfaces/OnlinePartyInterface.h
+++ b/Engine/Plugins/Online/OnlineSubsystem/Source/Public/Interfaces/OnlinePartyInterface.h
@@ -681,6 +681,13 @@ enum class EPartySystemState : uint8
  * @param Result Result of the operation
  */
 DECLARE_DELEGATE_TwoParams(FOnRestorePartiesComplete, const FUniqueNetId& /*LocalUserId*/, const FOnlineError& /*Result*/);
+/**
+ * Restore invites async task completed callback
+ *
+ * @param LocalUserId id of user that initiated the request
+ * @param Result Result of the operation
+ */
+DECLARE_DELEGATE_TwoParams(FOnRestoreInvitesComplete, const FUniqueNetId& /*LocalUserId*/, const FOnlineError& /*Result*/);
 /**
  * Cleanup parties async task completed callback
  *
@@ -1079,6 +1086,14 @@ public:
 	 * @param CompletionDelegate the delegate to trigger on completion
 	 */
 	virtual void RestoreParties(const FUniqueNetId& LocalUserId, const FOnRestorePartiesComplete& CompletionDelegate) = 0;
+
+	/**
+	 * Restore party invites. Intended to be called once during login to restore state from other running instances.
+	 *
+	 * @param LocalUserId the user to restore the pings for
+	 * @param CompletionDelegate the delegate to trigger on completion
+	 */
+	virtual void RestoreInvites(const FUniqueNetId& LocalUserId, const FOnRestoreInvitesComplete& CompletionDelegate) = 0;
 	
 	/**
 	 * Cleanup party state. This will cleanup the local party state and attempt to cleanup party memberships on an external service if possible.  Intended to be called for development purposes.
