Implementation of Settings File Compression/Decompression and Http Request to read/write in Async Task

diff --git a/Engine/Plugins/Online/OnlineSubsystem/Source/Public/Interfaces/OnlineUserCloudInterface.h b/Engine/Plugins/Online/OnlineSubsystem/Source/Public/Interfaces/OnlineUserCloudInterface.h
--- a/Engine/Plugins/Online/OnlineSubsystem/Source/Public/Interfaces/OnlineUserCloudInterface.h
+++ b/Engine/Plugins/Online/OnlineSubsystem/Source/Public/Interfaces/OnlineUserCloudInterface.h
@@ -188,10 +188,11 @@ public:
 	 * @param UserId User owning the storage
 	 * @param FileToWrite the name of the file to write
 	 * @param FileContents the out buffer to copy the data into
+	 * @param CompressBeforeUpload Whether the file will be compressed before upload. Defaulted to false. Use of this parameter isn't implemented in all platforms.
 	 *
 	 * @return true if the calls starts successfully, false otherwise
 	 */
-	virtual bool WriteUserFile(const FUniqueNetId& UserId, const FString& FileName, TArray<uint8>& FileContents) = 0;
+	virtual bool WriteUserFile(const FUniqueNetId& UserId, const FString& FileName, TArray<uint8>& FileContents, bool bCompressBeforeUpload = false) = 0;
 
 	/**
 	* Delegate fired at intervals during a user file write to the network platform's storage
