Add copy constructor to FAvatarInfo

diff --git a/Engine/Plugins/Online/OnlineSubsystem/Source/Public/Interfaces/AvatarProvider.h b/Engine/Plugins/Online/OnlineSubsystem/Source/Public/Interfaces/AvatarProvider.h
--- a/Engine/Plugins/Online/OnlineSubsystem/Source/Public/Interfaces/AvatarProvider.h
+++ b/Engine/Plugins/Online/OnlineSubsystem/Source/Public/Interfaces/AvatarProvider.h
@@ -19,6 +19,10 @@ public:
 		: UserId(InUserId.ToSharedRef()), AvatarInfoPairs(MoveTemp(InAvatarInfoPairs))
 	{ }
 
+	FAvatarInfo(const FAvatarInfo& AvatarInfo)
+		: UserId(AvatarInfo.UserId), AvatarInfoPairs(AvatarInfo.AvatarInfoPairs)
+	{ }
+
 	TSharedRef<const FUniqueNetId> UserId;
 	TMap<FString, FString> AvatarInfoPairs;
 };
