add party member spectator setting

#rnx

diff --git a/Engine/Plugins/Online/OnlineSubsystem/Source/Public/OnlineSessionSettings.h b/Engine/Plugins/Online/OnlineSubsystem/Source/Public/OnlineSessionSettings.h
--- a/Engine/Plugins/Online/OnlineSubsystem/Source/Public/OnlineSessionSettings.h
+++ b/Engine/Plugins/Online/OnlineSubsystem/Source/Public/OnlineSessionSettings.h
@@ -37,6 +37,8 @@
 #define SETTING_ALLOWBROADCASTING FName(TEXT("ALLOWBROADCASTING"))
 /** Setting defining broadcaster (value is bool) */
 #define SETTING_BROADCASTER FName(TEXT("BROADCASTER"))
+/** Setting defining if we are a party member spectator (value is bool) */
+#define SETTING_PARTYMEMBERSPECTATOR FName(TEXT("PARTYMEMBERSPECTATOR"))
 
 /** 8 user defined integer params to be used when filtering searches for sessions */
 #define SETTING_CUSTOMSEARCHINT1 FName(TEXT("CUSTOMSEARCHINT1"))
