FOnlineFriendsMcp implementation to add friend aliases
All queries are now pointing to V1 backend
Implemented SetFriendAlias.
All platforms that could be affected had been covered.
Some test command created to test functionality.

diff --git a/Engine/Plugins/Online/Android/OnlineSubsystemGameCircle/Source/OnlineSubsystemGameCircle/Private/OnlineFriendsInterfaceGameCircle.cpp b/Engine/Plugins/Online/Android/OnlineSubsystemGameCircle/Source/OnlineSubsystemGameCircle/Private/OnlineFriendsInterfaceGameCircle.cpp
--- a/Engine/Plugins/Online/Android/OnlineSubsystemGameCircle/Source/OnlineSubsystemGameCircle/Private/OnlineFriendsInterfaceGameCircle.cpp
+++ b/Engine/Plugins/Online/Android/OnlineSubsystemGameCircle/Source/OnlineSubsystemGameCircle/Private/OnlineFriendsInterfaceGameCircle.cpp
@@ -4,6 +4,7 @@
 #include "OnlineFriendGameCircle.h"
 #include "OnlineAGSPlayerClientCallbacks.h"
 #include "OnlineSubsystemGameCircle.h"
+#include "OnlineError.h"
 #include "OnlineIdentityInterfaceGameCircle.h"
 
 
@@ -54,6 +55,16 @@ bool FOnlineFriendsInterfaceGameCircle::RejectInvite(int32 LocalUserNum, const F
 {
 	return false;
 }
+
+void FOnlineFriendsInterfaceGameCircle::SetFriendAlias(int32 LocalUserNum, const FUniqueNetId& FriendId, const FString& ListName, const FString& Alias, const FOnSetFriendAliasComplete& Delegate = FOnSetFriendAliasComplete())
+{
+	TSharedRef<const FUniqueNetId> FriendIdRef = FriendId.AsShared();
+	GameCircleSubsystem->ExecuteNextTick([LocalUserNum, FriendIdRef, ListName, Delegate]()
+	{
+		Delegate.ExecuteIfBound(LocalUserNum, *FriendIdRef, ListName, FOnlineError(EOnlineErrorResult::NotImplemented));
+	});
+}
+
 bool FOnlineFriendsInterfaceGameCircle::DeleteFriend(int32 LocalUserNum, const FUniqueNetId& FriendId, const FString& ListName)
 {
 	return false;
diff --git a/Engine/Plugins/Online/Android/OnlineSubsystemGameCircle/Source/OnlineSubsystemGameCircle/Private/OnlineFriendsInterfaceGameCircle.h b/Engine/Plugins/Online/Android/OnlineSubsystemGameCircle/Source/OnlineSubsystemGameCircle/Private/OnlineFriendsInterfaceGameCircle.h
--- a/Engine/Plugins/Online/Android/OnlineSubsystemGameCircle/Source/OnlineSubsystemGameCircle/Private/OnlineFriendsInterfaceGameCircle.h
+++ b/Engine/Plugins/Online/Android/OnlineSubsystemGameCircle/Source/OnlineSubsystemGameCircle/Private/OnlineFriendsInterfaceGameCircle.h
@@ -19,6 +19,7 @@ public:
 	virtual bool SendInvite(int32 LocalUserNum, const FUniqueNetId& FriendId, const FString& ListName, const FOnSendInviteComplete& Delegate = FOnSendInviteComplete()) override;
 	virtual bool AcceptInvite(int32 LocalUserNum, const FUniqueNetId& FriendId, const FString& ListName, const FOnAcceptInviteComplete& Delegate = FOnAcceptInviteComplete()) override;
  	virtual bool RejectInvite(int32 LocalUserNum, const FUniqueNetId& FriendId, const FString& ListName) override;
+	virtual void SetFriendAlias(int32 LocalUserNum, const FUniqueNetId& FriendId, const FString& ListName, const FString& Alias, const FOnSetFriendAliasComplete& Delegate = FOnSetFriendAliasComplete()) override;
 	virtual bool DeleteFriend(int32 LocalUserNum, const FUniqueNetId& FriendId, const FString& ListName) override;
 	virtual bool GetFriendsList(int32 LocalUserNum, const FString& ListName, TArray< TSharedRef<FOnlineFriend> >& OutFriends) override;
 	virtual TSharedPtr<FOnlineFriend> GetFriend(int32 LocalUserNum, const FUniqueNetId& FriendId, const FString& ListName) override;
diff --git a/Engine/Plugins/Online/IOS/OnlineSubsystemIOS/Source/Private/OnlineFriendsInterfaceIOS.cpp b/Engine/Plugins/Online/IOS/OnlineSubsystemIOS/Source/Private/OnlineFriendsInterfaceIOS.cpp
--- a/Engine/Plugins/Online/IOS/OnlineSubsystemIOS/Source/Private/OnlineFriendsInterfaceIOS.cpp
+++ b/Engine/Plugins/Online/IOS/OnlineSubsystemIOS/Source/Private/OnlineFriendsInterfaceIOS.cpp
@@ -1,5 +1,6 @@
 // Copyright 1998-2019 Epic Games, Inc. All Rights Reserved.
 
+#include "OnlineError.h"
 #include "OnlineFriendsInterfaceIOS.h"
 #include "OnlineSubsystemIOS.h"
 
@@ -41,10 +42,10 @@ const FOnlineUserPresence& FOnlineFriendIOS::GetPresence() const
 
 // FOnlineFriendsIOS
 
-FOnlineFriendsIOS::FOnlineFriendsIOS(FOnlineSubsystemIOS* InSubsystem)
+FOnlineFriendsIOS::FOnlineFriendsIOS(FOnlineSubsystemIOS* InSubsystem):IOSSubsystem(InSubsystem)
 {
 	UE_LOG_ONLINE_FRIEND(Verbose, TEXT("FOnlineFriendsIOS::FOnlineFriendsIOS()"));
-
+	check(IOSSubsystem);
 	IdentityInterface = (FOnlineIdentityIOS*)InSubsystem->GetIdentityInterface().Get();
 }
 
@@ -229,6 +230,16 @@ bool FOnlineFriendsIOS::RejectInvite(int32 LocalUserNum, const FUniqueNetId& Fri
 	return false;
 }
 
+void FOnlineFriendsIOS::SetFriendAlias(int32 LocalUserNum, const FUniqueNetId& FriendId, const FString& ListName, const FString& Alias, const FOnSetFriendAliasComplete& Delegate = FOnSetFriendAliasComplete())
+{
+	TSharedRef<const FUniqueNetId> FriendIdRef = FriendId.AsShared();
+	IOSSubsystem.ExecuteNextTick([LocalUserNum, FriendIdRef, ListName, Delegate]()
+	{
+		UE_LOG_ONLINE_FRIEND(Warning, TEXT("FOnlineFriendsIOS::SetFriendAlias is currently not supported"));
+		Delegate.ExecuteIfBound(LocalUserNum, *FriendIdRef, ListName, FOnlineError(EOnlineErrorResult::NotImplemented));
+	});
+}
+
 bool FOnlineFriendsIOS::DeleteFriend(int32 LocalUserNum, const FUniqueNetId& FriendId, const FString& ListName)
 {
 	TriggerOnDeleteFriendCompleteDelegates(LocalUserNum, false, FriendId, ListName, FString(TEXT("DeleteFriend() is not supported")));
diff --git a/Engine/Plugins/Online/IOS/OnlineSubsystemIOS/Source/Private/OnlineFriendsInterfaceIOS.h b/Engine/Plugins/Online/IOS/OnlineSubsystemIOS/Source/Private/OnlineFriendsInterfaceIOS.h
--- a/Engine/Plugins/Online/IOS/OnlineSubsystemIOS/Source/Private/OnlineFriendsInterfaceIOS.h
+++ b/Engine/Plugins/Online/IOS/OnlineSubsystemIOS/Source/Private/OnlineFriendsInterfaceIOS.h
@@ -8,6 +8,8 @@
 #include "OnlineSubsystemIOSTypes.h"
 #include "Interfaces/OnlinePresenceInterface.h"
 
+class FOnlineSubsystemIOS;
+
 /**
  * Info associated with an online friend on the ios gamecenter service
  */
@@ -83,6 +85,7 @@ private:
 	/** The collection of game center friends received through the GK callbacks in ReadFriendsList */
 	TArray< TSharedRef<FOnlineFriendIOS> > CachedFriends;
 
+	FOnlineSubsystemIOS * IOSSubsystem;
 
 public:
 
@@ -93,6 +96,7 @@ public:
 	virtual bool SendInvite(int32 LocalUserNum, const FUniqueNetId& FriendId, const FString& ListName,  const FOnSendInviteComplete& Delegate = FOnSendInviteComplete()) override;
 	virtual bool AcceptInvite(int32 LocalUserNum, const FUniqueNetId& FriendId, const FString& ListName, const FOnAcceptInviteComplete& Delegate = FOnAcceptInviteComplete()) override;
  	virtual bool RejectInvite(int32 LocalUserNum, const FUniqueNetId& FriendId, const FString& ListName) override;
+	virtual void SetFriendAlias(int32 LocalUserNum, const FUniqueNetId& FriendId, const FString& ListName, const FString& Alias, const FOnSetFriendAliasComplete& Delegate = FOnSetFriendAliasComplete()) override;
  	virtual bool DeleteFriend(int32 LocalUserNum, const FUniqueNetId& FriendId, const FString& ListName) override;
 	virtual bool GetFriendsList(int32 LocalUserNum, const FString& ListName, TArray< TSharedRef<FOnlineFriend> >& OutFriends) override;
 	virtual TSharedPtr<FOnlineFriend> GetFriend(int32 LocalUserNum, const FUniqueNetId& FriendId, const FString& ListName) override;
diff --git a/Engine/Plugins/Online/OnlineSubsystem/Source/Public/Interfaces/OnlineFriendsInterface.h b/Engine/Plugins/Online/OnlineSubsystem/Source/Public/Interfaces/OnlineFriendsInterface.h
--- a/Engine/Plugins/Online/OnlineSubsystem/Source/Public/Interfaces/OnlineFriendsInterface.h
+++ b/Engine/Plugins/Online/OnlineSubsystem/Source/Public/Interfaces/OnlineFriendsInterface.h
@@ -165,6 +165,16 @@ DECLARE_DELEGATE_FiveParams(FOnAcceptInviteComplete, int32 /*LocalUserNum*/, boo
 DECLARE_MULTICAST_DELEGATE_FiveParams(FOnRejectInviteComplete, int32 /*LocalUserNum*/, bool /*bWasSuccessful*/, const FUniqueNetId& /*FriendId*/, const FString& /*ListName*/, const FString& /*ErrorStr*/);
 typedef FOnRejectInviteComplete::FDelegate FOnRejectInviteCompleteDelegate;
 
+/**
+ * Delegate used when a set alias request has completed
+ *
+ * @param LocalUserNum the controller number of the associated user that made the request
+ * @param FriendId player that we set alias for
+ * @param ListName name of the friends list that was operated on
+ * @param Error error information on failure
+ */
+DECLARE_DELEGATE_FourParams(FOnSetFriendAliasComplete, int32 /*LocalUserNum*/, const FUniqueNetId& /*FriendId*/, const FString& /*ListName*/, const FOnlineError& /*Error*/);
+
 /**
  * Delegate used when an friend delete request has completed
  *
@@ -437,6 +447,17 @@ public:
 	 */
  	virtual bool RejectInvite(int32 LocalUserNum, const FUniqueNetId& FriendId, const FString& ListName) = 0;
 
+	/**
+	 * Starts an async task that sets an alias a to friend.
+	 *
+	 * @param LocalUserNum the user that is setting the alias
+	 * @param FriendId player that will have the alias associated
+	 * @param ListName name of the friends list to operate on
+	 * @param Alias the alias associated to the player
+	 *
+	 */
+	virtual void SetFriendAlias(int32 LocalUserNum, const FUniqueNetId& FriendId, const FString& ListName, const FString& Alias, const FOnSetFriendAliasComplete& Delegate = FOnSetFriendAliasComplete()) = 0;
+	
 	/**
 	 * Delegate used when an invite reject request has completed
 	 *
diff --git a/Engine/Plugins/Online/OnlineSubsystem/Source/Public/OnlineSubsystemTypes.h b/Engine/Plugins/Online/OnlineSubsystem/Source/Public/OnlineSubsystemTypes.h
--- a/Engine/Plugins/Online/OnlineSubsystem/Source/Public/OnlineSubsystemTypes.h
+++ b/Engine/Plugins/Online/OnlineSubsystem/Source/Public/OnlineSubsystemTypes.h
@@ -1167,6 +1167,7 @@ struct FCloudFile
 #define USER_ATTR_PREFERRED_DISPLAYNAME TEXT("prefDisplayName")
 #define USER_ATTR_ID TEXT("id")
 #define USER_ATTR_EMAIL TEXT("email")
+#define USER_ATTR_ALIAS TEXT("alias")
 
 /**
  * Base for all online user info
diff --git a/Engine/Plugins/Online/OnlineSubsystemFacebook/Source/Private/OnlineFriendsFacebookCommon.cpp b/Engine/Plugins/Online/OnlineSubsystemFacebook/Source/Private/OnlineFriendsFacebookCommon.cpp
--- a/Engine/Plugins/Online/OnlineSubsystemFacebook/Source/Private/OnlineFriendsFacebookCommon.cpp
+++ b/Engine/Plugins/Online/OnlineSubsystemFacebook/Source/Private/OnlineFriendsFacebookCommon.cpp
@@ -7,6 +7,7 @@
 #else // USES_RESTFUL_FACEBOOK
 #include "OnlineIdentityFacebook.h"
 #endif // USES_RESTFUL_FACEBOOK
+#include "OnlineError.h"
 #include "HttpModule.h"
 #include "Interfaces/IHttpResponse.h"
 #include "Misc/ConfigCacheIni.h"
@@ -188,6 +189,16 @@ bool FOnlineFriendsFacebookCommon::RejectInvite(int32 LocalUserNum, const FUniqu
 	return false;
 }
 
+void FOnlineFriendsFacebookCommon::SetFriendAlias(int32 LocalUserNum, const FUniqueNetId& FriendId, const FString& ListName, const FString& Alias, const FOnSetFriendAliasComplete& Delegate /*= FOnSetFriendAliasComplete()*/)
+{
+	TSharedRef<const FUniqueNetId> FriendIdRef = FriendId.AsShared();
+	FacebookSubsystem->ExecuteNextTick([LocalUserNum, FriendIdRef, ListName, Delegate]()
+	{
+		UE_LOG_ONLINE_FRIEND(Warning, TEXT("FOnlineFriendsFacebookCommon::SetFriendAlias is not supported"));
+		Delegate.ExecuteIfBound(LocalUserNum, *FriendIdRef, ListName, FOnlineError(EOnlineErrorResult::NotImplemented));
+	});
+}
+
 bool FOnlineFriendsFacebookCommon::DeleteFriend(int32 LocalUserNum, const FUniqueNetId& FriendId, const FString& ListName)
 {
 	TriggerOnDeleteFriendCompleteDelegates(LocalUserNum, false, FriendId, ListName, FString(TEXT("DeleteFriend() is not supported")));
diff --git a/Engine/Plugins/Online/OnlineSubsystemFacebook/Source/Private/OnlineFriendsFacebookCommon.h b/Engine/Plugins/Online/OnlineSubsystemFacebook/Source/Private/OnlineFriendsFacebookCommon.h
--- a/Engine/Plugins/Online/OnlineSubsystemFacebook/Source/Private/OnlineFriendsFacebookCommon.h
+++ b/Engine/Plugins/Online/OnlineSubsystemFacebook/Source/Private/OnlineFriendsFacebookCommon.h
@@ -123,6 +123,7 @@ public:
 	virtual bool SendInvite(int32 LocalUserNum, const FUniqueNetId& FriendId, const FString& ListName,  const FOnSendInviteComplete& Delegate = FOnSendInviteComplete()) override;
 	virtual bool AcceptInvite(int32 LocalUserNum, const FUniqueNetId& FriendId, const FString& ListName, const FOnAcceptInviteComplete& Delegate = FOnAcceptInviteComplete()) override;
  	virtual bool RejectInvite(int32 LocalUserNum, const FUniqueNetId& FriendId, const FString& ListName) override;
+	virtual void SetFriendAlias(int32 LocalUserNum, const FUniqueNetId& FriendId, const FString& ListName, const FString& Alias, const FOnSetFriendAliasComplete& Delegate = FOnSetFriendAliasComplete()) override;
  	virtual bool DeleteFriend(int32 LocalUserNum, const FUniqueNetId& FriendId, const FString& ListName) override;
 	virtual bool GetFriendsList(int32 LocalUserNum, const FString& ListName, TArray< TSharedRef<FOnlineFriend> >& OutFriends) override;
 	virtual TSharedPtr<FOnlineFriend> GetFriend(int32 LocalUserNum, const FUniqueNetId& FriendId, const FString& ListName) override;
diff --git a/Engine/Plugins/Online/OnlineSubsystemOculus/Source/Private/OnlineFriendsInterfaceOculus.cpp b/Engine/Plugins/Online/OnlineSubsystemOculus/Source/Private/OnlineFriendsInterfaceOculus.cpp
--- a/Engine/Plugins/Online/OnlineSubsystemOculus/Source/Private/OnlineFriendsInterfaceOculus.cpp
+++ b/Engine/Plugins/Online/OnlineSubsystemOculus/Source/Private/OnlineFriendsInterfaceOculus.cpp
@@ -2,6 +2,7 @@
 
 #include "OnlineFriendsInterfaceOculus.h"
 #include "OnlineSubsystemOculusPrivate.h"
+#include "OnlineError.h"
 #include "OnlineSubsystemOculusPackage.h"
 
 FOnlineFriendsOculus::FOnlineFriendsOculus(FOnlineSubsystemOculus& InSubsystem)
@@ -117,6 +118,16 @@ bool FOnlineFriendsOculus::RejectInvite(int32 LocalUserNum, const FUniqueNetId&
 	return false;
 }
 
+void FOnlineFriendsOculus::SetFriendAlias(int32 LocalUserNum, const FUniqueNetId& FriendId, const FString& ListName, const FString& Alias, const FOnSetFriendAliasComplete& Delegate = FOnSetFriendAliasComplete())
+{
+	TSharedRef<const FUniqueNetId> FriendIdRef = FriendId.AsShared();
+	OculusSubsystem.ExecuteNextTick([LocalUserNum, FriendIdRef, ListName, Delegate]()
+	{
+		UE_LOG_ONLINE_FRIEND(Warning, TEXT("FOnlineFriendsOculus::SetFriendAlias is not implemented"));
+		Delegate.ExecuteIfBound(LocalUserNum, *FriendIdRef, ListName, FOnlineError(EOnlineErrorResult::NotImplemented));
+	});
+}
+
 bool FOnlineFriendsOculus::DeleteFriend(int32 LocalUserNum, const FUniqueNetId& FriendId, const FString& ListName)
 {
 	/** Does not exist in LibOVRPlatform */
diff --git a/Engine/Plugins/Online/OnlineSubsystemOculus/Source/Private/OnlineFriendsInterfaceOculus.h b/Engine/Plugins/Online/OnlineSubsystemOculus/Source/Private/OnlineFriendsInterfaceOculus.h
--- a/Engine/Plugins/Online/OnlineSubsystemOculus/Source/Private/OnlineFriendsInterfaceOculus.h
+++ b/Engine/Plugins/Online/OnlineSubsystemOculus/Source/Private/OnlineFriendsInterfaceOculus.h
@@ -111,6 +111,7 @@ public:
 	virtual bool SendInvite(int32 LocalUserNum, const FUniqueNetId& FriendId, const FString& ListName, const FOnSendInviteComplete& Delegate = FOnSendInviteComplete()) override;
 	virtual bool AcceptInvite(int32 LocalUserNum, const FUniqueNetId& FriendId, const FString& ListName, const FOnAcceptInviteComplete& Delegate = FOnAcceptInviteComplete()) override;
 	virtual bool RejectInvite(int32 LocalUserNum, const FUniqueNetId& FriendId, const FString& ListName) override;
+	virtual void SetFriendAlias(int32 LocalUserNum, const FUniqueNetId& FriendId, const FString& ListName, const FString& Alias, const FOnSetFriendAliasComplete& Delegate = FOnSetFriendAliasComplete()) override;
 	virtual bool DeleteFriend(int32 LocalUserNum, const FUniqueNetId& FriendId, const FString& ListName) override;
 	virtual bool GetFriendsList(int32 LocalUserNum, const FString& ListName, TArray< TSharedRef<FOnlineFriend> >& OutFriends) override;
 	virtual TSharedPtr<FOnlineFriend> GetFriend(int32 LocalUserNum, const FUniqueNetId& FriendId, const FString& ListName) override;
diff --git a/Engine/Plugins/Online/OnlineSubsystemSteam/Source/Private/OnlineFriendsInterfaceSteam.cpp b/Engine/Plugins/Online/OnlineSubsystemSteam/Source/Private/OnlineFriendsInterfaceSteam.cpp
--- a/Engine/Plugins/Online/OnlineSubsystemSteam/Source/Private/OnlineFriendsInterfaceSteam.cpp
+++ b/Engine/Plugins/Online/OnlineSubsystemSteam/Source/Private/OnlineFriendsInterfaceSteam.cpp
@@ -2,6 +2,7 @@
 
 #include "OnlineFriendsInterfaceSteam.h"
 #include "OnlineSubsystemSteam.h"
+#include "OnlineError.h"
 #include "OnlineSubsystemSteamTypes.h"
 
 // FOnlineFriendSteam
@@ -102,6 +103,16 @@ bool FOnlineFriendsSteam::RejectInvite(int32 LocalUserNum, const FUniqueNetId& F
 	return false;
 }
 
+void FOnlineFriendsSteam::SetFriendAlias(int32 LocalUserNum, const FUniqueNetId& FriendId, const FString& ListName, const FString& Alias, const FOnSetFriendAliasComplete& Delegate = FOnSetFriendAliasComplete())
+{
+	TSharedRef<const FUniqueNetId> FriendIdRef = FriendId.AsShared();
+	SteamSubsystem->ExecuteNextTick([LocalUserNum, FriendIdRef, ListName, Delegate]()
+	{
+		UE_LOG_ONLINE_FRIEND(Warning, TEXT("FOnlineFriendsSteam::SetFriendAlias is not supported"));
+		Delegate.ExecuteIfBound(LocalUserNum, *FriendIdRef, ListName, FOnlineError(EOnlineErrorResult::NotImplemented));
+	});
+}
+
 bool FOnlineFriendsSteam::DeleteFriend(int32 LocalUserNum, const FUniqueNetId& FriendId, const FString& ListName)
 {
 	TriggerOnDeleteFriendCompleteDelegates(LocalUserNum, false, FriendId, ListName, FString(TEXT("DeleteFriend() is not supported")));
diff --git a/Engine/Plugins/Online/OnlineSubsystemSteam/Source/Private/OnlineFriendsInterfaceSteam.h b/Engine/Plugins/Online/OnlineSubsystemSteam/Source/Private/OnlineFriendsInterfaceSteam.h
--- a/Engine/Plugins/Online/OnlineSubsystemSteam/Source/Private/OnlineFriendsInterfaceSteam.h
+++ b/Engine/Plugins/Online/OnlineSubsystemSteam/Source/Private/OnlineFriendsInterfaceSteam.h
@@ -114,7 +114,8 @@ public:
 	virtual bool SendInvite(int32 LocalUserNum, const FUniqueNetId& FriendId, const FString& ListName,  const FOnSendInviteComplete& Delegate = FOnSendInviteComplete()) override;
 	virtual bool AcceptInvite(int32 LocalUserNum, const FUniqueNetId& FriendId, const FString& ListName, const FOnAcceptInviteComplete& Delegate = FOnAcceptInviteComplete()) override;
  	virtual bool RejectInvite(int32 LocalUserNum, const FUniqueNetId& FriendId, const FString& ListName) override;
- 	virtual bool DeleteFriend(int32 LocalUserNum, const FUniqueNetId& FriendId, const FString& ListName) override;
+	virtual void SetFriendAlias(int32 LocalUserNum, const FUniqueNetId& FriendId, const FString& ListName, const FString& Alias, const FOnSetFriendAliasComplete& Delegate = FOnSetFriendAliasComplete()) override;
+	virtual bool DeleteFriend(int32 LocalUserNum, const FUniqueNetId& FriendId, const FString& ListName) override;
 	virtual bool GetFriendsList(int32 LocalUserNum, const FString& ListName, TArray< TSharedRef<FOnlineFriend> >& OutFriends) override;
 	virtual TSharedPtr<FOnlineFriend> GetFriend(int32 LocalUserNum, const FUniqueNetId& FriendId, const FString& ListName) override;
 	virtual bool IsFriend(int32 LocalUserNum, const FUniqueNetId& FriendId, const FString& ListName) override;
diff --git a/Engine/Plugins/Online/OnlineSubsystemTencent/Source/Private/OnlineFriendsTencent.cpp b/Engine/Plugins/Online/OnlineSubsystemTencent/Source/Private/OnlineFriendsTencent.cpp
--- a/Engine/Plugins/Online/OnlineSubsystemTencent/Source/Private/OnlineFriendsTencent.cpp
+++ b/Engine/Plugins/Online/OnlineSubsystemTencent/Source/Private/OnlineFriendsTencent.cpp
@@ -469,6 +469,16 @@ bool FOnlineFriendsTencent::RejectInvite(int32 LocalUserNum, const FUniqueNetId&
 	return bStarted;
 }
 
+void FOnlineFriendsTencent::SetFriendAlias(int32 LocalUserNum, const FUniqueNetId& FriendId, const FString& ListName, const FString& Alias, const FOnSetFriendAliasComplete& Delegate /*= FOnSetFriendAliasComplete()*/)
+{
+	TSharedRef<const FUniqueNetId> FriendIdRef = FriendId.AsShared();
+	TencentSubsystem->ExecuteNextTick([LocalUserNum, FriendIdRef, ListName, Delegate]()
+	{
+		UE_LOG_ONLINE_FRIEND(Warning, TEXT("FOnlineFriendsTencent::SetFriendAlias is not implemented"));
+		Delegate.ExecuteIfBound(LocalUserNum, *FriendIdRef, ListName, FOnlineError(EOnlineErrorResult::NotImplemented));
+	});
+}
+
 bool FOnlineFriendsTencent::DeleteFriend(int32 LocalUserNum, const FUniqueNetId& FriendId, const FString& ListName)
 {
 	/** NYI */
diff --git a/Engine/Plugins/Online/OnlineSubsystemTencent/Source/Private/OnlineFriendsTencent.h b/Engine/Plugins/Online/OnlineSubsystemTencent/Source/Private/OnlineFriendsTencent.h
--- a/Engine/Plugins/Online/OnlineSubsystemTencent/Source/Private/OnlineFriendsTencent.h
+++ b/Engine/Plugins/Online/OnlineSubsystemTencent/Source/Private/OnlineFriendsTencent.h
@@ -187,6 +187,7 @@ public:
 	virtual bool SendInvite(int32 LocalUserNum, const FUniqueNetId& FriendId, const FString& ListName, const FOnSendInviteComplete& Delegate = FOnSendInviteComplete()) override;
 	virtual bool AcceptInvite(int32 LocalUserNum, const FUniqueNetId& FriendId, const FString& ListName, const FOnAcceptInviteComplete& Delegate = FOnAcceptInviteComplete()) override;
 	virtual bool RejectInvite(int32 LocalUserNum, const FUniqueNetId& FriendId, const FString& ListName) override;
+	virtual void SetFriendAlias(int32 LocalUserNum, const FUniqueNetId& FriendId, const FString& ListName, const FString& Alias, const FOnSetFriendAliasComplete& Delegate = FOnSetFriendAliasComplete()) override;
 	virtual bool DeleteFriend(int32 LocalUserNum, const FUniqueNetId& FriendId, const FString& ListName) override;
 	virtual bool GetFriendsList(int32 LocalUserNum, const FString& ListName, TArray< TSharedRef<FOnlineFriend> >& OutFriends) override;
 	virtual TSharedPtr<FOnlineFriend> GetFriend(int32 LocalUserNum, const FUniqueNetId& FriendId, const FString& ListName) override;
