Fix typo in doc comment

diff --git a/Engine/Plugins/Online/OnlineSubsystem/Source/Public/Interfaces/OnlinePartyInterface.h b/Engine/Plugins/Online/OnlineSubsystem/Source/Public/Interfaces/OnlinePartyInterface.h
--- a/Engine/Plugins/Online/OnlineSubsystem/Source/Public/Interfaces/OnlinePartyInterface.h
+++ b/Engine/Plugins/Online/OnlineSubsystem/Source/Public/Interfaces/OnlinePartyInterface.h
@@ -608,7 +608,7 @@ public:
 	TSharedPtr<const FUniqueNetId> LeaderId;
 	/** The current state of the party */
 	EPartyState State;
-	/** The current state of the party */
+	/** The previous state of the party */
 	EPartyState PreviousState;
 	/** id of chat room associated with the party */
 	FChatRoomId RoomId;
