Added comparator for FAvatarInfo.

#tests Compile DebugGame Editor Win64 / Tests using multiple -game clients

diff --git a/Engine/Plugins/Online/OnlineSubsystem/Source/Public/Interfaces/AvatarProvider.h b/Engine/Plugins/Online/OnlineSubsystem/Source/Public/Interfaces/AvatarProvider.h
--- a/Engine/Plugins/Online/OnlineSubsystem/Source/Public/Interfaces/AvatarProvider.h
+++ b/Engine/Plugins/Online/OnlineSubsystem/Source/Public/Interfaces/AvatarProvider.h
@@ -23,6 +23,36 @@ public:
 		: UserId(AvatarInfo.UserId), AvatarInfoPairs(AvatarInfo.AvatarInfoPairs)
 	{ }
 
+	bool Equals(const FAvatarInfo& Other) const
+	{
+		if (UserId.Get() != Other.UserId.Get())
+		{
+			return false;
+		}
+		if (AvatarInfoPairs.Num() != Other.AvatarInfoPairs.Num())
+		{
+			return false;
+		}
+
+		for (const TPair<FString, FString>& Pair : AvatarInfoPairs)
+		{
+			if (const FString* OtherValue = Other.AvatarInfoPairs.Find(Pair.Key))
+			{
+				if (!Pair.Value.Equals(*OtherValue))
+				{
+					return false;
+				}
+			}
+			else
+			{
+				return false;
+			}
+			
+		}
+
+		return true;
+	}
+
 	TSharedRef<const FUniqueNetId> UserId;
 	TMap<FString, FString> AvatarInfoPairs;
 };
