Add XMPP, websockets, VoiceChatManager, HTTP, and MCP to a ticker that runs in the foreground _and_ background

diff --git a/Engine/Plugins/Online/OnlineSubsystem/Source/Private/OnlineSubsystemImpl.cpp b/Engine/Plugins/Online/OnlineSubsystem/Source/Private/OnlineSubsystemImpl.cpp
--- a/Engine/Plugins/Online/OnlineSubsystem/Source/Private/OnlineSubsystemImpl.cpp
+++ b/Engine/Plugins/Online/OnlineSubsystem/Source/Private/OnlineSubsystemImpl.cpp
@@ -31,17 +31,18 @@ namespace OSSConsoleVariables
 const FName FOnlineSubsystemImpl::DefaultInstanceName(TEXT("DefaultInstance"));
 
 FOnlineSubsystemImpl::FOnlineSubsystemImpl(FName InSubsystemName, FName InInstanceName) :
+	FOnlineSubsystemImpl(InSubsystemName, InInstanceName, FTicker::GetCoreTicker())
+{
+}
+
+FOnlineSubsystemImpl::FOnlineSubsystemImpl(FName InSubsystemName, FName InInstanceName, FTicker& Ticker) :
+	FTickerObjectBase(0.0f, Ticker),
 	SubsystemName(InSubsystemName),
 	InstanceName(InInstanceName),
 	bForceDedicated(false),
-	NamedInterfaces(nullptr)
+	NamedInterfaces(nullptr),
+	bTickerStarted(true)
 {
-	StartTicker();
-}
-
-FOnlineSubsystemImpl::~FOnlineSubsystemImpl()
-{	
-	ensure(!TickHandle.IsValid());
 }
 
 void FOnlineSubsystemImpl::PreUnload()
@@ -94,43 +95,36 @@ void FOnlineSubsystemImpl::ExecuteDelegateNextTick(const FNextTickDelegate& Call
 
 void FOnlineSubsystemImpl::StartTicker()
 {
-	if (!TickHandle.IsValid())
-	{
-		// Register delegate for ticker callback
-		FTickerDelegate TickDelegate = FTickerDelegate::CreateRaw(this, &FOnlineSubsystemImpl::Tick);
-		TickHandle = FTicker::GetCoreTicker().AddTicker(TickDelegate, 0.0f);
-	}
+	bTickerStarted = true;
 }
 
 void FOnlineSubsystemImpl::StopTicker()
 {
-	// Unregister ticker delegate
-	if (TickHandle.IsValid())
-	{
-		FTicker::GetCoreTicker().RemoveTicker(TickHandle);
-		TickHandle.Reset();
-	}
+	bTickerStarted = false;
 }
 
 bool FOnlineSubsystemImpl::Tick(float DeltaTime)
 {
-	QUICK_SCOPE_CYCLE_COUNTER(STAT_FOnlineSubsystemImpl_Tick);
-	if (!NextTickQueue.IsEmpty())
+	if (bTickerStarted)
 	{
-		// unload the next-tick queue into our buffer. Any further executes (from within callbacks) will happen NEXT frame (as intended)
-		FNextTickDelegate Temp;
-		while (NextTickQueue.Dequeue(Temp))
+		QUICK_SCOPE_CYCLE_COUNTER(STAT_FOnlineSubsystemImpl_Tick);
+		if (!NextTickQueue.IsEmpty())
 		{
-			CurrentTickBuffer.Add(Temp);
-		}
+			// unload the next-tick queue into our buffer. Any further executes (from within callbacks) will happen NEXT frame (as intended)
+			FNextTickDelegate Temp;
+			while (NextTickQueue.Dequeue(Temp))
+			{
+				CurrentTickBuffer.Add(Temp);
+			}
 
-		// execute any functions in the current tick array
-		for (const auto& Callback : CurrentTickBuffer)
-		{
-			QUICK_SCOPE_CYCLE_COUNTER(STAT_FOnlineSubsystemImpl_Tick_ExecuteCallback);
-			Callback.ExecuteIfBound();
+			// execute any functions in the current tick array
+			for (const auto& Callback : CurrentTickBuffer)
+			{
+				QUICK_SCOPE_CYCLE_COUNTER(STAT_FOnlineSubsystemImpl_Tick_ExecuteCallback);
+				Callback.ExecuteIfBound();
+			}
+			CurrentTickBuffer.SetNum(0, false); // keep the memory around
 		}
-		CurrentTickBuffer.SetNum(0, false); // keep the memory around
 	}
 	return true;
 }
diff --git a/Engine/Plugins/Online/OnlineSubsystem/Source/Public/OnlineSubsystemImpl.h b/Engine/Plugins/Online/OnlineSubsystem/Source/Public/OnlineSubsystemImpl.h
--- a/Engine/Plugins/Online/OnlineSubsystem/Source/Public/OnlineSubsystemImpl.h
+++ b/Engine/Plugins/Online/OnlineSubsystem/Source/Public/OnlineSubsystemImpl.h
@@ -7,6 +7,7 @@
 #include "OnlineSubsystem.h"
 #include "Containers/Queue.h"
 #include "OnlineSubsystemPackage.h"
+#include "Containers/Ticker.h"
 
 struct FOnlineError;
 
@@ -22,6 +23,7 @@ namespace OSSConsoleVariables
  */
 class ONLINESUBSYSTEM_API FOnlineSubsystemImpl 
 	: public IOnlineSubsystem
+	, public FTickerObjectBase
 {
 private:
 
@@ -45,7 +47,9 @@ protected:
 
 	/** Hidden on purpose */
 	FOnlineSubsystemImpl() = delete;
+
 	FOnlineSubsystemImpl(FName InSubsystemName, FName InInstanceName);
+	FOnlineSubsystemImpl(FName InSubsystemName, FName InInstanceName, FTicker& Ticker);
 
 	/** Name of the subsystem @see OnlineSubsystemNames.h */
 	FName SubsystemName;
@@ -75,13 +79,13 @@ protected:
 
 	/** Stop Ticker */
 	void StopTicker();
-	
-	/** Delegate for callbacks to Tick */
-	FDelegateHandle TickHandle;
+
+	/** Is the ticker started */
+	bool bTickerStarted;
 
 public:
 	
-	virtual ~FOnlineSubsystemImpl();
+	virtual ~FOnlineSubsystemImpl() = default;
 
 	// IOnlineSubsystem
 	virtual void PreUnload() override;
@@ -102,13 +106,8 @@ public:
 	virtual bool IsEnabled() const override;
 	virtual void ReloadConfigs(const TSet<FString>& /*ConfigSections*/) override {};
 
-	/**
-	 * Tick function
-	 *
-	 * @param DeltaTime	time passed since the last call.
-	 * @return true if should continue ticking
-	 */
-	virtual bool Tick(float DeltaTime);
+	// FTickerObjectBase
+	virtual bool Tick(float DeltaTime) override;
 
 	/**
 	 * Modify a response string so that it can be logged cleanly
diff --git a/Engine/Source/Runtime/Core/Private/Containers/BackgroundableTicker.cpp b/Engine/Source/Runtime/Core/Private/Containers/BackgroundableTicker.cpp
new file mode 100644
--- /dev/null
+++ b/Engine/Source/Runtime/Core/Private/Containers/BackgroundableTicker.cpp
@@ -0,0 +1,36 @@
+// Copyright 1998-2019 Epic Games, Inc. All Rights Reserved.
+
+#include "Containers/BackgroundableTicker.h"
+#include "Misc/CoreDelegates.h"
+
+FBackgroundableTicker& FBackgroundableTicker::GetCoreTicker()
+{
+	static FBackgroundableTicker Singleton;
+	return Singleton;
+}
+
+FBackgroundableTicker::FBackgroundableTicker()
+{
+	CoreTickerHandle = FTicker::GetCoreTicker().AddTicker(TEXT("FBackgroundableTicker"), 0.0f, [this](float DeltaTime) -> bool
+	{
+		if (bWasBackgrounded)
+		{
+			// When returning to the foreground, ensure we do not report enormous delta time values coming from the foreground ticker.
+			bWasBackgrounded = false;
+			DeltaTime = FMath::Clamp(DeltaTime, 0.0f, 1.0f / 60);
+		}
+		Tick(DeltaTime);
+		return true;
+	});
+	BackgroundTickerHandle = FCoreDelegates::MobileBackgroundTickDelegate.AddLambda([this](float DeltaTime)
+	{
+		bWasBackgrounded = true;
+		Tick(DeltaTime);
+	});
+}
+
+FBackgroundableTicker::~FBackgroundableTicker()
+{
+	FTicker::GetCoreTicker().RemoveTicker(CoreTickerHandle);
+	FCoreDelegates::MobileBackgroundTickDelegate.Remove(BackgroundTickerHandle);
+}
\ No newline at end of file
diff --git a/Engine/Source/Runtime/Core/Private/Misc/CoreDelegates.cpp b/Engine/Source/Runtime/Core/Private/Misc/CoreDelegates.cpp
--- a/Engine/Source/Runtime/Core/Private/Misc/CoreDelegates.cpp
+++ b/Engine/Source/Runtime/Core/Private/Misc/CoreDelegates.cpp
@@ -100,6 +100,7 @@ FCoreDelegates::FApplicationLifetimeDelegate FCoreDelegates::ApplicationWillEnte
 FCoreDelegates::FApplicationLifetimeDelegate FCoreDelegates::ApplicationHasEnteredForegroundDelegate;
 FCoreDelegates::FApplicationLifetimeDelegate FCoreDelegates::ApplicationWillTerminateDelegate;
 FCoreDelegates::FApplicationLifetimeDelegate FCoreDelegates::ApplicationShouldUnloadResourcesDelegate;
+FCoreDelegates::FBackgroundTickDelegate FCoreDelegates::MobileBackgroundTickDelegate;
 
 FCoreDelegates::FApplicationStartupArgumentsDelegate FCoreDelegates::ApplicationReceivedStartupArgumentsDelegate;
 
diff --git a/Engine/Source/Runtime/Core/Public/Containers/BackgroundableTicker.h b/Engine/Source/Runtime/Core/Public/Containers/BackgroundableTicker.h
new file mode 100644
--- /dev/null
+++ b/Engine/Source/Runtime/Core/Public/Containers/BackgroundableTicker.h
@@ -0,0 +1,22 @@
+// Copyright 1998-2019 Epic Games, Inc. All Rights Reserved.
+
+#pragma once
+
+#include "Ticker.h"
+
+class FBackgroundableTicker
+	: public FTicker
+{
+public:
+
+	CORE_API static FBackgroundableTicker& GetCoreTicker();
+
+	CORE_API FBackgroundableTicker();
+	CORE_API ~FBackgroundableTicker();
+
+private:
+	
+	FDelegateHandle CoreTickerHandle;
+	FDelegateHandle BackgroundTickerHandle;
+	bool bWasBackgrounded = false;
+};
\ No newline at end of file
diff --git a/Engine/Source/Runtime/Core/Public/Containers/Ticker.h b/Engine/Source/Runtime/Core/Public/Containers/Ticker.h
--- a/Engine/Source/Runtime/Core/Public/Containers/Ticker.h
+++ b/Engine/Source/Runtime/Core/Public/Containers/Ticker.h
@@ -72,7 +72,7 @@ public:
 	CORE_API void Tick(float DeltaTime);
 
 
-private:
+protected:
 	/**
 	 * Element of the priority queue
 	 */
@@ -158,4 +158,4 @@ private:
 	FTicker& Ticker;
 	/** Delegate for callbacks to Tick */
 	FDelegateHandle TickHandle;
-};
+};
\ No newline at end of file
diff --git a/Engine/Source/Runtime/Core/Public/Misc/CoreDelegates.h b/Engine/Source/Runtime/Core/Public/Misc/CoreDelegates.h
--- a/Engine/Source/Runtime/Core/Public/Misc/CoreDelegates.h
+++ b/Engine/Source/Runtime/Core/Public/Misc/CoreDelegates.h
@@ -339,7 +339,7 @@ public:
 	// terminated from the background state without any further warning.
 	static FApplicationLifetimeDelegate ApplicationWillEnterBackgroundDelegate; // for instance, hitting the home button
 
-																				// Called when the application is returning to the foreground (reverse any processing done in the EnterBackground delegate)
+	// Called when the application is returning to the foreground (reverse any processing done in the EnterBackground delegate)
 	static FApplicationLifetimeDelegate ApplicationHasEnteredForegroundDelegate;
 
 	// This *may* be called when the application is getting terminated by the OS.
@@ -347,6 +347,12 @@ public:
 	// save state when ApplicationWillEnterBackgroundDelegate is called instead.
 	static FApplicationLifetimeDelegate ApplicationWillTerminateDelegate;
 
+	// Called when in the background, if the OS is giving CPU time to the device. It is very likely
+	// this will never be called due to mobile OS backgrounded CPU restrictions. But if, for instance,
+	// VOIP is active on iOS, the will be getting called
+	DECLARE_MULTICAST_DELEGATE_OneParam(FBackgroundTickDelegate, float /*DeltaTime*/);
+	static FBackgroundTickDelegate MobileBackgroundTickDelegate;
+
 	// Called when the OS needs control of the music (parameter is true) or when the OS returns
 	// control of the music to the application (parameter is false). This can happen due to a
 	// phone call or timer or other OS-level event. This is currently triggered only on iOS
diff --git a/Engine/Source/Runtime/Launch/Private/IOS/LaunchIOS.cpp b/Engine/Source/Runtime/Launch/Private/IOS/LaunchIOS.cpp
--- a/Engine/Source/Runtime/Launch/Private/IOS/LaunchIOS.cpp
+++ b/Engine/Source/Runtime/Launch/Private/IOS/LaunchIOS.cpp
@@ -418,16 +418,29 @@ void FAppEntry::Init()
 #endif // UE_BUILD_SHIPPING
 }
 
+#if BUILD_EMBEDDED_APP
+static bool GWasTickSuspended = false;
+static double GPreviousSuspendTime = FPlatformTime::Seconds();
+#else
 static FSuspendRenderingThread* SuspendThread = NULL;
+#endif
 
 void FAppEntry::Tick()
 {
-    if (SuspendThread != NULL)
-    {
-        delete SuspendThread;
-        SuspendThread = NULL;
-        FPlatformProcess::SetRealTimeMode();
-    }
+#if BUILD_EMBEDDED_APP
+	if (GWasTickSuspended)
+	{
+		FPlatformProcess::SetRealTimeMode();
+		GWasTickSuspended = false;
+	}
+#else
+	if (SuspendThread != NULL)
+	{
+		delete SuspendThread;
+		SuspendThread = NULL;
+		FPlatformProcess::SetRealTimeMode();
+	}
+#endif
     
 	if (AudioContextResumeTime != 0)
 	{
@@ -444,12 +457,29 @@ void FAppEntry::Tick()
 
 void FAppEntry::SuspendTick()
 {
-    if (!SuspendThread)
+#if BUILD_EMBEDDED_APP
+	static double PreviousTime = FPlatformTime::Seconds();
+    if (!GWasTickSuspended)
     {
-        SuspendThread = new FSuspendRenderingThread(true);
+		GWasTickSuspended = true;
+		// reset it each time we background
+		PreviousTime = FPlatformTime::Seconds();
     }
-    
-    FPlatformProcess::Sleep(0.1f);
+
+	float DeltaTime = FPlatformTime::Seconds() - PreviousTime;
+	PreviousTime = FPlatformTime::Seconds();
+
+	// allow for some background processing
+	FEmbeddedCommunication::TickGameThread(DeltaTime);
+	FCoreDelegates::MobileBackgroundTickDelegate.Broadcast(DeltaTime);
+#else
+	if (!SuspendThread)
+	{
+		SuspendThread = new FSuspendRenderingThread(true);
+	}
+#endif
+
+	FPlatformProcess::Sleep(0.1f);
 }
 
 void FAppEntry::Shutdown()
diff --git a/Engine/Source/Runtime/Online/HTTP/Private/HttpManager.cpp b/Engine/Source/Runtime/Online/HTTP/Private/HttpManager.cpp
--- a/Engine/Source/Runtime/Online/HTTP/Private/HttpManager.cpp
+++ b/Engine/Source/Runtime/Online/HTTP/Private/HttpManager.cpp
@@ -13,13 +13,14 @@
 #include "Misc/CommandLine.h"
 
 #include "Stats/Stats.h"
+#include "Containers/BackgroundableTicker.h"
 
 // FHttpManager
 
 FCriticalSection FHttpManager::RequestLock;
 
 FHttpManager::FHttpManager()
-	: FTickerObjectBase(0.0f)
+	: FTickerObjectBase(0.0f, FBackgroundableTicker::GetCoreTicker())
 	, Thread(nullptr)
 	, CorrelationIdMethod(FHttpManager::GetDefaultCorrelationIdMethod())
 {
diff --git a/Engine/Source/Runtime/Online/WebSockets/Private/Lws/LwsWebSocketsManager.cpp b/Engine/Source/Runtime/Online/WebSockets/Private/Lws/LwsWebSocketsManager.cpp
--- a/Engine/Source/Runtime/Online/WebSockets/Private/Lws/LwsWebSocketsManager.cpp
+++ b/Engine/Source/Runtime/Online/WebSockets/Private/Lws/LwsWebSocketsManager.cpp
@@ -9,13 +9,15 @@
 #if WITH_SSL
 #include "Ssl.h"
 #endif
+#include "HttpModule.h"
 #include "WebSocketsLog.h"
+#include "Containers/BackgroundableTicker.h"
 #include "HAL/PlatformTime.h"
 #include "HAL/PlatformProcess.h"
 #include "HAL/LowLevelMemTracker.h"
 #include "Misc/ConfigCacheIni.h"
+#include "Misc/CoreDelegates.h"
 #include "Stats/Stats.h"
-#include "HttpModule.h"
 
 namespace {
 	static const struct lws_extension LwsExtensions[] = {
@@ -158,17 +160,15 @@ void FLwsWebSocketsManager::InitWebSockets(TArrayView<const FString> Protocols)
 	}
 
 	// Setup our game thread tick
-	FTicker& Ticker = FTicker::GetCoreTicker();
 	FTickerDelegate TickDelegate = FTickerDelegate::CreateRaw(this, &FLwsWebSocketsManager::GameThreadTick);
-	TickHandle = Ticker.AddTicker(TickDelegate, 0.0f);
+	TickHandle = FBackgroundableTicker::GetCoreTicker().AddTicker(TickDelegate, 0.0f);
 }
 
 void FLwsWebSocketsManager::ShutdownWebSockets()
 {
 	if (TickHandle.IsValid())
 	{
-		FTicker& Ticker = FTicker::GetCoreTicker();
-		Ticker.RemoveTicker(TickHandle);
+		FBackgroundableTicker::GetCoreTicker().RemoveTicker(TickHandle);
 		TickHandle.Reset();
 	}
 
diff --git a/Engine/Source/Runtime/Online/XMPP/Private/XmppJingle/XmppChatJingle.cpp b/Engine/Source/Runtime/Online/XMPP/Private/XmppJingle/XmppChatJingle.cpp
--- a/Engine/Source/Runtime/Online/XMPP/Private/XmppJingle/XmppChatJingle.cpp
+++ b/Engine/Source/Runtime/Online/XMPP/Private/XmppJingle/XmppChatJingle.cpp
@@ -1,6 +1,7 @@
 // Copyright 1998-2019 Epic Games, Inc. All Rights Reserved.
 
 #include "XmppJingle/XmppChatJingle.h"
+#include "Containers/BackgroundableTicker.h"
 #include "XmppJingle/XmppConnectionJingle.h"
 #include "XmppPresence.h"
 #include "XmppLog.h"
@@ -197,7 +198,8 @@ protected:
 };
 
 FXmppChatJingle::FXmppChatJingle(class FXmppConnectionJingle& InConnection)
-	: ChatRcvTask(NULL)
+	: FTickerObjectBase(0.0f, FBackgroundableTicker::GetCoreTicker())
+	, ChatRcvTask(NULL)
 	, ChatSendTask(NULL)
 	, Connection(InConnection)
 {
diff --git a/Engine/Source/Runtime/Online/XMPP/Private/XmppJingle/XmppConnectionJingle.cpp b/Engine/Source/Runtime/Online/XMPP/Private/XmppJingle/XmppConnectionJingle.cpp
--- a/Engine/Source/Runtime/Online/XMPP/Private/XmppJingle/XmppConnectionJingle.cpp
+++ b/Engine/Source/Runtime/Online/XMPP/Private/XmppJingle/XmppConnectionJingle.cpp
@@ -1,6 +1,7 @@
 // Copyright 1998-2019 Epic Games, Inc. All Rights Reserved.
 
 #include "XmppJingle/XmppConnectionJingle.h"
+#include "Containers/BackgroundableTicker.h"
 #include "HAL/Runnable.h"
 #include "HAL/RunnableThread.h"
 #include "HAL/ThreadSafeCounter.h"
@@ -404,7 +405,8 @@ private:
 int32 FXmppConnectionPumpThread::ThreadInstanceIdx = 0;
 
 FXmppConnectionJingle::FXmppConnectionJingle()
-	: LastLoginState(EXmppLoginStatus::NotStarted)
+	: FTickerObjectBase(0.0f, FBackgroundableTicker::GetCoreTicker())
+	, LastLoginState(EXmppLoginStatus::NotStarted)
 	, LoginState(EXmppLoginStatus::NotStarted)
 	, StatUpdateFreq(1.0)
 	, LastStatUpdateTime(0.0)
diff --git a/Engine/Source/Runtime/Online/XMPP/Private/XmppJingle/XmppMessagesJingle.cpp b/Engine/Source/Runtime/Online/XMPP/Private/XmppJingle/XmppMessagesJingle.cpp
--- a/Engine/Source/Runtime/Online/XMPP/Private/XmppJingle/XmppMessagesJingle.cpp
+++ b/Engine/Source/Runtime/Online/XMPP/Private/XmppJingle/XmppMessagesJingle.cpp
@@ -17,8 +17,8 @@ class FXmppMessageJingle
 {
 public:
 	FXmppMessageJingle(
-		const buzz::Jid& InFromJid = buzz::Jid(), 
-		const buzz::Jid& InToJid = buzz::Jid(), 
+		const buzz::Jid& InFromJid = buzz::Jid(),
+		const buzz::Jid& InToJid = buzz::Jid(),
 		const std::string& InBody = std::string()
 		)
 		: FromJid(InFromJid)
@@ -185,7 +185,8 @@ protected:
 };
 
 FXmppMessagesJingle::FXmppMessagesJingle(class FXmppConnectionJingle& InConnection)
-	: MessageRcvTask(NULL)
+	: FTickerObjectBase(0.0f, FBackgroundableTicker::GetCoreTicker())
+	, MessageRcvTask(NULL)
 	, MessageSendTask(NULL)
 	, NumMessagesReceived(0)
 	, NumMessagesSent(0)
diff --git a/Engine/Source/Runtime/Online/XMPP/Private/XmppJingle/XmppMultiUserChatJingle.cpp b/Engine/Source/Runtime/Online/XMPP/Private/XmppJingle/XmppMultiUserChatJingle.cpp
--- a/Engine/Source/Runtime/Online/XMPP/Private/XmppJingle/XmppMultiUserChatJingle.cpp
+++ b/Engine/Source/Runtime/Online/XMPP/Private/XmppJingle/XmppMultiUserChatJingle.cpp
@@ -214,7 +214,8 @@ private:
 };
 
 FXmppMultiUserChatJingle::FXmppMultiUserChatJingle(class FXmppConnectionJingle& InConnection)
-	: Connection(InConnection)
+	: FTickerObjectBase(0.0f, FBackgroundableTicker::GetCoreTicker())
+	, Connection(InConnection)
 	, NumOpRequests(0)
 	, NumMucResponses(0)
 	, VerbosityIncreasedCount(0)
diff --git a/Engine/Source/Runtime/Online/XMPP/Private/XmppJingle/XmppPresenceJingle.cpp b/Engine/Source/Runtime/Online/XMPP/Private/XmppJingle/XmppPresenceJingle.cpp
--- a/Engine/Source/Runtime/Online/XMPP/Private/XmppJingle/XmppPresenceJingle.cpp
+++ b/Engine/Source/Runtime/Online/XMPP/Private/XmppJingle/XmppPresenceJingle.cpp
@@ -5,6 +5,7 @@
 #include "XmppLog.h"
 #include "XmppJingle/XmppConnectionJingle.h"
 #include "XmppMultiUserChat.h"
+#include "Containers/BackgroundableTicker.h"
 
 #if WITH_XMPP_JINGLE
 
@@ -356,7 +357,8 @@ private:
 };
 
 FXmppPresenceJingle::FXmppPresenceJingle(class FXmppConnectionJingle& InConnection)
-	: PresenceSendTask(NULL)
+	: FTickerObjectBase(0.0f, FBackgroundableTicker::GetCoreTicker())
+	, PresenceSendTask(NULL)
 	, PresenceRcvTask(NULL)
 	, NumPresenceIn(0)
 	, NumPresenceOut(0)
diff --git a/Engine/Source/Runtime/Online/XMPP/Private/XmppJingle/XmppPubSubJingle.cpp b/Engine/Source/Runtime/Online/XMPP/Private/XmppJingle/XmppPubSubJingle.cpp
--- a/Engine/Source/Runtime/Online/XMPP/Private/XmppJingle/XmppPubSubJingle.cpp
+++ b/Engine/Source/Runtime/Online/XMPP/Private/XmppJingle/XmppPubSubJingle.cpp
@@ -4,6 +4,7 @@
 #include "XmppJingle/XmppConnectionJingle.h"
 #include "XmppLog.h"
 #include "Misc/ScopeLock.h"
+#include "Containers/BackgroundableTicker.h"
 
 #if WITH_XMPP_JINGLE
 
@@ -133,7 +134,8 @@ bool FXmppPubSubJingle::GetLastMessages(const FXmppPubSubId& NodeId, int32 NumMe
 }
 
 FXmppPubSubJingle::FXmppPubSubJingle(class FXmppConnectionJingle& InConnection)
-	: Connection(InConnection)
+	: FTickerObjectBase(0.0f, FBackgroundableTicker::GetCoreTicker())
+	, Connection(InConnection)
 {
 
 }
diff --git a/Engine/Source/Runtime/Online/XMPP/Private/XmppStrophe/XmppConnectionStrophe.cpp b/Engine/Source/Runtime/Online/XMPP/Private/XmppStrophe/XmppConnectionStrophe.cpp
--- a/Engine/Source/Runtime/Online/XMPP/Private/XmppStrophe/XmppConnectionStrophe.cpp
+++ b/Engine/Source/Runtime/Online/XMPP/Private/XmppStrophe/XmppConnectionStrophe.cpp
@@ -13,11 +13,13 @@
 #include "XmppStrophe/StropheStanzaConstants.h"
 #include "XmppStrophe/StropheError.h"
 #include "XmppLog.h"
+#include "Containers/BackgroundableTicker.h"
 
 #if WITH_XMPP_STROPHE
 
 FXmppConnectionStrophe::FXmppConnectionStrophe()
-	: LoginStatus(EXmppLoginStatus::NotStarted)
+	: FTickerObjectBase(0.0f, FBackgroundableTicker::GetCoreTicker())
+	, LoginStatus(EXmppLoginStatus::NotStarted)
 {
 	MessagesStrophe = MakeShared<FXmppMessagesStrophe, ESPMode::ThreadSafe>(*this);
 	MultiUserChatStrophe = MakeShared<FXmppMultiUserChatStrophe, ESPMode::ThreadSafe>(*this);
diff --git a/Engine/Source/Runtime/Online/XMPP/Private/XmppStrophe/XmppMessagesStrophe.cpp b/Engine/Source/Runtime/Online/XMPP/Private/XmppStrophe/XmppMessagesStrophe.cpp
--- a/Engine/Source/Runtime/Online/XMPP/Private/XmppStrophe/XmppMessagesStrophe.cpp
+++ b/Engine/Source/Runtime/Online/XMPP/Private/XmppStrophe/XmppMessagesStrophe.cpp
@@ -15,13 +15,15 @@
 #include "Policies/CondensedJsonPrintPolicy.h"
 #include "Misc/Guid.h"
 #include "Misc/EmbeddedCommunication.h"
+#include "Containers/BackgroundableTicker.h"
 
 #if WITH_XMPP_STROPHE
 
 #define TickRequesterId FName("StropheMessages")
 
 FXmppMessagesStrophe::FXmppMessagesStrophe(FXmppConnectionStrophe& InConnectionManager)
-	: ConnectionManager(InConnectionManager)
+	: FTickerObjectBase(0.0f, FBackgroundableTicker::GetCoreTicker())
+	, ConnectionManager(InConnectionManager)
 {
 }
 
diff --git a/Engine/Source/Runtime/Online/XMPP/Private/XmppStrophe/XmppMultiUserChatStrophe.cpp b/Engine/Source/Runtime/Online/XMPP/Private/XmppStrophe/XmppMultiUserChatStrophe.cpp
--- a/Engine/Source/Runtime/Online/XMPP/Private/XmppStrophe/XmppMultiUserChatStrophe.cpp
+++ b/Engine/Source/Runtime/Online/XMPP/Private/XmppStrophe/XmppMultiUserChatStrophe.cpp
@@ -10,13 +10,15 @@
 #include "Logging/LogScopedVerbosityOverride.h"
 #include "Misc/Guid.h"
 #include "Misc/EmbeddedCommunication.h"
+#include "Containers/BackgroundableTicker.h"
 
 #if WITH_XMPP_STROPHE
 
 #define TickRequesterId FName("StropheMultichat")
 
 FXmppMultiUserChatStrophe::FXmppMultiUserChatStrophe(FXmppConnectionStrophe& InConnectionManager)
-	: ConnectionManager(InConnectionManager)
+	: FTickerObjectBase(0.0f, FBackgroundableTicker::GetCoreTicker())
+	, ConnectionManager(InConnectionManager)
 {
 
 }
diff --git a/Engine/Source/Runtime/Online/XMPP/Private/XmppStrophe/XmppPingStrophe.cpp b/Engine/Source/Runtime/Online/XMPP/Private/XmppStrophe/XmppPingStrophe.cpp
--- a/Engine/Source/Runtime/Online/XMPP/Private/XmppStrophe/XmppPingStrophe.cpp
+++ b/Engine/Source/Runtime/Online/XMPP/Private/XmppStrophe/XmppPingStrophe.cpp
@@ -7,13 +7,15 @@
 #include "XmppLog.h"
 #include "Misc/Guid.h"
 #include "Misc/EmbeddedCommunication.h"
+#include "Containers/BackgroundableTicker.h"
 
 #if WITH_XMPP_STROPHE
 
 #define TickRequesterId FName("StrophePing")
 
 FXmppPingStrophe::FXmppPingStrophe(FXmppConnectionStrophe& InConnectionManager)
-	: ConnectionManager(InConnectionManager)
+	: FTickerObjectBase(0.0f, FBackgroundableTicker::GetCoreTicker())
+	, ConnectionManager(InConnectionManager)
 	, TimeSinceLastClientPing(0.0f)
 	, bWaitingForPong(false)
 	, SecondsSinceLastServerPong(0.0f)
diff --git a/Engine/Source/Runtime/Online/XMPP/Private/XmppStrophe/XmppPresenceStrophe.cpp b/Engine/Source/Runtime/Online/XMPP/Private/XmppStrophe/XmppPresenceStrophe.cpp
--- a/Engine/Source/Runtime/Online/XMPP/Private/XmppStrophe/XmppPresenceStrophe.cpp
+++ b/Engine/Source/Runtime/Online/XMPP/Private/XmppStrophe/XmppPresenceStrophe.cpp
@@ -5,13 +5,15 @@
 #include "XmppStrophe/StropheStanza.h"
 #include "XmppStrophe/StropheStanzaConstants.h"
 #include "Misc/EmbeddedCommunication.h"
+#include "Containers/BackgroundableTicker.h"
 
 #if WITH_XMPP_STROPHE
 
 #define TickRequesterId FName("StrophePresence")
 
 FXmppPresenceStrophe::FXmppPresenceStrophe(FXmppConnectionStrophe& InConnectionManager)
-	: ConnectionManager(InConnectionManager)
+	: FTickerObjectBase(0.0f, FBackgroundableTicker::GetCoreTicker())
+	, ConnectionManager(InConnectionManager)
 {
 
 }
diff --git a/Engine/Source/Runtime/Online/XMPP/Private/XmppStrophe/XmppPrivateChatStrophe.cpp b/Engine/Source/Runtime/Online/XMPP/Private/XmppStrophe/XmppPrivateChatStrophe.cpp
--- a/Engine/Source/Runtime/Online/XMPP/Private/XmppStrophe/XmppPrivateChatStrophe.cpp
+++ b/Engine/Source/Runtime/Online/XMPP/Private/XmppStrophe/XmppPrivateChatStrophe.cpp
@@ -7,13 +7,15 @@
 #include "XmppStrophe/StropheStanzaConstants.h"
 #include "XmppLog.h"
 #include "Misc/EmbeddedCommunication.h"
+#include "Containers/BackgroundableTicker.h"
 
 #if WITH_XMPP_STROPHE
 
 #define TickRequesterId FName("StrophePrivateChat")
 
 FXmppPrivateChatStrophe::FXmppPrivateChatStrophe(FXmppConnectionStrophe& InConnectionManager)
-	: ConnectionManager(InConnectionManager)
+	: FTickerObjectBase(0.0f, FBackgroundableTicker::GetCoreTicker())
+	, ConnectionManager(InConnectionManager)
 {
 }
 
diff --git a/Engine/Source/Runtime/Online/XMPP/Private/XmppStrophe/XmppPubSubStrophe.cpp b/Engine/Source/Runtime/Online/XMPP/Private/XmppStrophe/XmppPubSubStrophe.cpp
--- a/Engine/Source/Runtime/Online/XMPP/Private/XmppStrophe/XmppPubSubStrophe.cpp
+++ b/Engine/Source/Runtime/Online/XMPP/Private/XmppStrophe/XmppPubSubStrophe.cpp
@@ -2,11 +2,13 @@
 
 #include "XmppPubSubStrophe.h"
 #include "XmppConnectionStrophe.h"
+#include "Containers/BackgroundableTicker.h"
 
 #if WITH_XMPP_STROPHE
 
 FXmppPubSubStrophe::FXmppPubSubStrophe(FXmppConnectionStrophe& InConnectionManager)
-	: ConnectionManager(InConnectionManager)
+	: FTickerObjectBase(0.0f, FBackgroundableTicker::GetCoreTicker())
+	, ConnectionManager(InConnectionManager)
 {
 
 }
