Created DeleteFriendAlias to remove a nickname from friends that have one.
All related platforms updated

diff --git a/Engine/Plugins/Online/Android/OnlineSubsystemGameCircle/Source/OnlineSubsystemGameCircle/Private/OnlineFriendsInterfaceGameCircle.cpp b/Engine/Plugins/Online/Android/OnlineSubsystemGameCircle/Source/OnlineSubsystemGameCircle/Private/OnlineFriendsInterfaceGameCircle.cpp
--- a/Engine/Plugins/Online/Android/OnlineSubsystemGameCircle/Source/OnlineSubsystemGameCircle/Private/OnlineFriendsInterfaceGameCircle.cpp
+++ b/Engine/Plugins/Online/Android/OnlineSubsystemGameCircle/Source/OnlineSubsystemGameCircle/Private/OnlineFriendsInterfaceGameCircle.cpp
@@ -65,6 +65,15 @@ void FOnlineFriendsInterfaceGameCircle::SetFriendAlias(int32 LocalUserNum, const
 	});
 }
 
+void FOnlineFriendsInterfaceGameCircle::DeleteFriendAlias(int32 LocalUserNum, const FUniqueNetId& FriendId, const FString& ListName, const FOnDeleteFriendAliasComplete& Delegate)
+{
+	TSharedRef<const FUniqueNetId> FriendIdRef = FriendId.AsShared();
+	GameCircleSubsystem->ExecuteNextTick([LocalUserNum, FriendIdRef, ListName, Delegate]()
+	{
+		Delegate.ExecuteIfBound(LocalUserNum, *FriendIdRef, ListName, FOnlineError(EOnlineErrorResult::NotImplemented));
+	});
+}
+
 bool FOnlineFriendsInterfaceGameCircle::DeleteFriend(int32 LocalUserNum, const FUniqueNetId& FriendId, const FString& ListName)
 {
 	return false;
diff --git a/Engine/Plugins/Online/Android/OnlineSubsystemGameCircle/Source/OnlineSubsystemGameCircle/Private/OnlineFriendsInterfaceGameCircle.h b/Engine/Plugins/Online/Android/OnlineSubsystemGameCircle/Source/OnlineSubsystemGameCircle/Private/OnlineFriendsInterfaceGameCircle.h
--- a/Engine/Plugins/Online/Android/OnlineSubsystemGameCircle/Source/OnlineSubsystemGameCircle/Private/OnlineFriendsInterfaceGameCircle.h
+++ b/Engine/Plugins/Online/Android/OnlineSubsystemGameCircle/Source/OnlineSubsystemGameCircle/Private/OnlineFriendsInterfaceGameCircle.h
@@ -20,6 +20,7 @@ public:
 	virtual bool AcceptInvite(int32 LocalUserNum, const FUniqueNetId& FriendId, const FString& ListName, const FOnAcceptInviteComplete& Delegate = FOnAcceptInviteComplete()) override;
  	virtual bool RejectInvite(int32 LocalUserNum, const FUniqueNetId& FriendId, const FString& ListName) override;
 	virtual void SetFriendAlias(int32 LocalUserNum, const FUniqueNetId& FriendId, const FString& ListName, const FString& Alias, const FOnSetFriendAliasComplete& Delegate = FOnSetFriendAliasComplete()) override;
+	virtual void DeleteFriendAlias(int32 LocalUserNum, const FUniqueNetId& FriendId, const FString& ListName, const FOnDeleteFriendAliasComplete& Delegate = FOnDeleteFriendAliasComplete()) override;
 	virtual bool DeleteFriend(int32 LocalUserNum, const FUniqueNetId& FriendId, const FString& ListName) override;
 	virtual bool GetFriendsList(int32 LocalUserNum, const FString& ListName, TArray< TSharedRef<FOnlineFriend> >& OutFriends) override;
 	virtual TSharedPtr<FOnlineFriend> GetFriend(int32 LocalUserNum, const FUniqueNetId& FriendId, const FString& ListName) override;
diff --git a/Engine/Plugins/Online/IOS/OnlineSubsystemIOS/Source/Private/OnlineFriendsInterfaceIOS.cpp b/Engine/Plugins/Online/IOS/OnlineSubsystemIOS/Source/Private/OnlineFriendsInterfaceIOS.cpp
--- a/Engine/Plugins/Online/IOS/OnlineSubsystemIOS/Source/Private/OnlineFriendsInterfaceIOS.cpp
+++ b/Engine/Plugins/Online/IOS/OnlineSubsystemIOS/Source/Private/OnlineFriendsInterfaceIOS.cpp
@@ -241,6 +241,16 @@ void FOnlineFriendsIOS::SetFriendAlias(int32 LocalUserNum, const FUniqueNetId& F
 	});
 }
 
+void FOnlineFriendsIOS::DeleteFriendAlias(int32 LocalUserNum, const FUniqueNetId& FriendId, const FString& ListName, const FOnDeleteFriendAliasComplete& Delegate)
+{
+	TSharedRef<const FUniqueNetId> FriendIdRef = FriendId.AsShared();
+	IOSSubsystem->ExecuteNextTick([LocalUserNum, FriendIdRef, ListName, Delegate]()
+	{
+		UE_LOG_ONLINE_FRIEND(Warning, TEXT("FOnlineFriendsIOS::DeleteFriendAlias is currently not supported"));
+		Delegate.ExecuteIfBound(LocalUserNum, *FriendIdRef, ListName, FOnlineError(EOnlineErrorResult::NotImplemented));
+	});
+}
+
 bool FOnlineFriendsIOS::DeleteFriend(int32 LocalUserNum, const FUniqueNetId& FriendId, const FString& ListName)
 {
 	TriggerOnDeleteFriendCompleteDelegates(LocalUserNum, false, FriendId, ListName, FString(TEXT("DeleteFriend() is not supported")));
diff --git a/Engine/Plugins/Online/IOS/OnlineSubsystemIOS/Source/Private/OnlineFriendsInterfaceIOS.h b/Engine/Plugins/Online/IOS/OnlineSubsystemIOS/Source/Private/OnlineFriendsInterfaceIOS.h
--- a/Engine/Plugins/Online/IOS/OnlineSubsystemIOS/Source/Private/OnlineFriendsInterfaceIOS.h
+++ b/Engine/Plugins/Online/IOS/OnlineSubsystemIOS/Source/Private/OnlineFriendsInterfaceIOS.h
@@ -97,7 +97,8 @@ public:
 	virtual bool AcceptInvite(int32 LocalUserNum, const FUniqueNetId& FriendId, const FString& ListName, const FOnAcceptInviteComplete& Delegate = FOnAcceptInviteComplete()) override;
  	virtual bool RejectInvite(int32 LocalUserNum, const FUniqueNetId& FriendId, const FString& ListName) override;
 	virtual void SetFriendAlias(int32 LocalUserNum, const FUniqueNetId& FriendId, const FString& ListName, const FString& Alias, const FOnSetFriendAliasComplete& Delegate = FOnSetFriendAliasComplete()) override;
- 	virtual bool DeleteFriend(int32 LocalUserNum, const FUniqueNetId& FriendId, const FString& ListName) override;
+	virtual void DeleteFriendAlias(int32 LocalUserNum, const FUniqueNetId& FriendId, const FString& ListName, const FOnDeleteFriendAliasComplete& Delegate = FOnDeleteFriendAliasComplete()) override; 
+	virtual bool DeleteFriend(int32 LocalUserNum, const FUniqueNetId& FriendId, const FString& ListName) override;
 	virtual bool GetFriendsList(int32 LocalUserNum, const FString& ListName, TArray< TSharedRef<FOnlineFriend> >& OutFriends) override;
 	virtual TSharedPtr<FOnlineFriend> GetFriend(int32 LocalUserNum, const FUniqueNetId& FriendId, const FString& ListName) override;
 	virtual bool IsFriend(int32 LocalUserNum, const FUniqueNetId& FriendId, const FString& ListName) override;
diff --git a/Engine/Plugins/Online/OnlineSubsystem/Source/Public/Interfaces/OnlineFriendsInterface.h b/Engine/Plugins/Online/OnlineSubsystem/Source/Public/Interfaces/OnlineFriendsInterface.h
--- a/Engine/Plugins/Online/OnlineSubsystem/Source/Public/Interfaces/OnlineFriendsInterface.h
+++ b/Engine/Plugins/Online/OnlineSubsystem/Source/Public/Interfaces/OnlineFriendsInterface.h
@@ -175,6 +175,16 @@ typedef FOnRejectInviteComplete::FDelegate FOnRejectInviteCompleteDelegate;
  */
 DECLARE_DELEGATE_FourParams(FOnSetFriendAliasComplete, int32 /*LocalUserNum*/, const FUniqueNetId& /*FriendId*/, const FString& /*ListName*/, const FOnlineError& /*Error*/);
 
+/**
+ * Delegate used when a delete friend alias request has completed
+ *
+ * @param LocalUserNum the controller number of the associated user that made the request
+ * @param FriendId player that we delete alias from
+ * @param ListName name of the friends list that was operated on
+ * @param Error error information on failure
+ */
+DECLARE_DELEGATE_FourParams(FOnDeleteFriendAliasComplete, int32 /*LocalUserNum*/, const FUniqueNetId& /*FriendId*/, const FString& /*ListName*/, const FOnlineError& /*Error*/);
+
 /**
  * Delegate used when an friend delete request has completed
  *
@@ -457,7 +467,17 @@ public:
 	 *
 	 */
 	virtual void SetFriendAlias(int32 LocalUserNum, const FUniqueNetId& FriendId, const FString& ListName, const FString& Alias, const FOnSetFriendAliasComplete& Delegate = FOnSetFriendAliasComplete()) = 0;
-	
+
+	/**
+	 * Starts an async task that delete an alias from a friend.
+	 *
+	 * @param LocalUserNum the user that is setting the alias
+	 * @param FriendId player that will have the alias deleted from
+	 * @param ListName name of the friends list to operate on
+	 *
+	 */
+	virtual void DeleteFriendAlias(int32 LocalUserNum, const FUniqueNetId& FriendId, const FString& ListName, const FOnDeleteFriendAliasComplete& Delegate = FOnDeleteFriendAliasComplete()) = 0;
+
 	/**
 	 * Delegate used when an invite reject request has completed
 	 *
diff --git a/Engine/Plugins/Online/OnlineSubsystemFacebook/Source/Private/OnlineFriendsFacebookCommon.cpp b/Engine/Plugins/Online/OnlineSubsystemFacebook/Source/Private/OnlineFriendsFacebookCommon.cpp
--- a/Engine/Plugins/Online/OnlineSubsystemFacebook/Source/Private/OnlineFriendsFacebookCommon.cpp
+++ b/Engine/Plugins/Online/OnlineSubsystemFacebook/Source/Private/OnlineFriendsFacebookCommon.cpp
@@ -199,6 +199,16 @@ void FOnlineFriendsFacebookCommon::SetFriendAlias(int32 LocalUserNum, const FUni
 	});
 }
 
+void FOnlineFriendsFacebookCommon::DeleteFriendAlias(int32 LocalUserNum, const FUniqueNetId& FriendId, const FString& ListName, const FOnDeleteFriendAliasComplete& Delegate)
+{
+	TSharedRef<const FUniqueNetId> FriendIdRef = FriendId.AsShared();
+	FacebookSubsystem->ExecuteNextTick([LocalUserNum, FriendIdRef, ListName, Delegate]()
+	{
+		UE_LOG_ONLINE_FRIEND(Warning, TEXT("FOnlineFriendsFacebookCommon::DeleteFriendAlias is not supported"));
+		Delegate.ExecuteIfBound(LocalUserNum, *FriendIdRef, ListName, FOnlineError(EOnlineErrorResult::NotImplemented));
+	});
+}
+
 bool FOnlineFriendsFacebookCommon::DeleteFriend(int32 LocalUserNum, const FUniqueNetId& FriendId, const FString& ListName)
 {
 	TriggerOnDeleteFriendCompleteDelegates(LocalUserNum, false, FriendId, ListName, FString(TEXT("DeleteFriend() is not supported")));
diff --git a/Engine/Plugins/Online/OnlineSubsystemFacebook/Source/Private/OnlineFriendsFacebookCommon.h b/Engine/Plugins/Online/OnlineSubsystemFacebook/Source/Private/OnlineFriendsFacebookCommon.h
--- a/Engine/Plugins/Online/OnlineSubsystemFacebook/Source/Private/OnlineFriendsFacebookCommon.h
+++ b/Engine/Plugins/Online/OnlineSubsystemFacebook/Source/Private/OnlineFriendsFacebookCommon.h
@@ -124,7 +124,8 @@ public:
 	virtual bool AcceptInvite(int32 LocalUserNum, const FUniqueNetId& FriendId, const FString& ListName, const FOnAcceptInviteComplete& Delegate = FOnAcceptInviteComplete()) override;
  	virtual bool RejectInvite(int32 LocalUserNum, const FUniqueNetId& FriendId, const FString& ListName) override;
 	virtual void SetFriendAlias(int32 LocalUserNum, const FUniqueNetId& FriendId, const FString& ListName, const FString& Alias, const FOnSetFriendAliasComplete& Delegate = FOnSetFriendAliasComplete()) override;
- 	virtual bool DeleteFriend(int32 LocalUserNum, const FUniqueNetId& FriendId, const FString& ListName) override;
+	virtual void DeleteFriendAlias(int32 LocalUserNum, const FUniqueNetId& FriendId, const FString& ListName, const FOnDeleteFriendAliasComplete& Delegate = FOnDeleteFriendAliasComplete()) override;
+	virtual bool DeleteFriend(int32 LocalUserNum, const FUniqueNetId& FriendId, const FString& ListName) override;
 	virtual bool GetFriendsList(int32 LocalUserNum, const FString& ListName, TArray< TSharedRef<FOnlineFriend> >& OutFriends) override;
 	virtual TSharedPtr<FOnlineFriend> GetFriend(int32 LocalUserNum, const FUniqueNetId& FriendId, const FString& ListName) override;
 	virtual bool IsFriend(int32 LocalUserNum, const FUniqueNetId& FriendId, const FString& ListName) override;
diff --git a/Engine/Plugins/Online/OnlineSubsystemOculus/Source/Private/OnlineFriendsInterfaceOculus.cpp b/Engine/Plugins/Online/OnlineSubsystemOculus/Source/Private/OnlineFriendsInterfaceOculus.cpp
--- a/Engine/Plugins/Online/OnlineSubsystemOculus/Source/Private/OnlineFriendsInterfaceOculus.cpp
+++ b/Engine/Plugins/Online/OnlineSubsystemOculus/Source/Private/OnlineFriendsInterfaceOculus.cpp
@@ -128,6 +128,15 @@ void FOnlineFriendsOculus::SetFriendAlias(int32 LocalUserNum, const FUniqueNetId
 	});
 }
 
+void FOnlineFriendsOculus::DeleteFriendAlias(int32 LocalUserNum, const FUniqueNetId& FriendId, const FString& ListName, const FOnDeleteFriendAliasComplete& Delegate)
+{
+	TSharedRef<const FUniqueNetId> FriendIdRef = FriendId.AsShared();
+	OculusSubsystem.ExecuteNextTick([LocalUserNum, FriendIdRef, ListName, Delegate]()
+	{
+		UE_LOG_ONLINE_FRIEND(Warning, TEXT("FOnlineFriendsOculus::DeleteFriendAlias is not implemented"));
+		Delegate.ExecuteIfBound(LocalUserNum, *FriendIdRef, ListName, FOnlineError(EOnlineErrorResult::NotImplemented));
+	});
+}
 bool FOnlineFriendsOculus::DeleteFriend(int32 LocalUserNum, const FUniqueNetId& FriendId, const FString& ListName)
 {
 	/** Does not exist in LibOVRPlatform */
diff --git a/Engine/Plugins/Online/OnlineSubsystemOculus/Source/Private/OnlineFriendsInterfaceOculus.h b/Engine/Plugins/Online/OnlineSubsystemOculus/Source/Private/OnlineFriendsInterfaceOculus.h
--- a/Engine/Plugins/Online/OnlineSubsystemOculus/Source/Private/OnlineFriendsInterfaceOculus.h
+++ b/Engine/Plugins/Online/OnlineSubsystemOculus/Source/Private/OnlineFriendsInterfaceOculus.h
@@ -112,6 +112,7 @@ public:
 	virtual bool AcceptInvite(int32 LocalUserNum, const FUniqueNetId& FriendId, const FString& ListName, const FOnAcceptInviteComplete& Delegate = FOnAcceptInviteComplete()) override;
 	virtual bool RejectInvite(int32 LocalUserNum, const FUniqueNetId& FriendId, const FString& ListName) override;
 	virtual void SetFriendAlias(int32 LocalUserNum, const FUniqueNetId& FriendId, const FString& ListName, const FString& Alias, const FOnSetFriendAliasComplete& Delegate = FOnSetFriendAliasComplete()) override;
+	virtual void DeleteFriendAlias(int32 LocalUserNum, const FUniqueNetId& FriendId, const FString& ListName, const FOnDeleteFriendAliasComplete& Delegate = FOnDeleteFriendAliasComplete()) override;
 	virtual bool DeleteFriend(int32 LocalUserNum, const FUniqueNetId& FriendId, const FString& ListName) override;
 	virtual bool GetFriendsList(int32 LocalUserNum, const FString& ListName, TArray< TSharedRef<FOnlineFriend> >& OutFriends) override;
 	virtual TSharedPtr<FOnlineFriend> GetFriend(int32 LocalUserNum, const FUniqueNetId& FriendId, const FString& ListName) override;
diff --git a/Engine/Plugins/Online/OnlineSubsystemSteam/Source/Private/OnlineFriendsInterfaceSteam.cpp b/Engine/Plugins/Online/OnlineSubsystemSteam/Source/Private/OnlineFriendsInterfaceSteam.cpp
--- a/Engine/Plugins/Online/OnlineSubsystemSteam/Source/Private/OnlineFriendsInterfaceSteam.cpp
+++ b/Engine/Plugins/Online/OnlineSubsystemSteam/Source/Private/OnlineFriendsInterfaceSteam.cpp
@@ -113,6 +113,16 @@ void FOnlineFriendsSteam::SetFriendAlias(int32 LocalUserNum, const FUniqueNetId&
 	});
 }
 
+void FOnlineFriendsSteam::DeleteFriendAlias(int32 LocalUserNum, const FUniqueNetId& FriendId, const FString& ListName, const FOnDeleteFriendAliasComplete& Delegate)
+{
+	TSharedRef<const FUniqueNetId> FriendIdRef = FriendId.AsShared();
+	SteamSubsystem->ExecuteNextTick([LocalUserNum, FriendIdRef, ListName, Delegate]()
+	{
+		UE_LOG_ONLINE_FRIEND(Warning, TEXT("FOnlineFriendsSteam::DeleteFriendAlias is not supported"));
+		Delegate.ExecuteIfBound(LocalUserNum, *FriendIdRef, ListName, FOnlineError(EOnlineErrorResult::NotImplemented));
+	});
+}
+
 bool FOnlineFriendsSteam::DeleteFriend(int32 LocalUserNum, const FUniqueNetId& FriendId, const FString& ListName)
 {
 	TriggerOnDeleteFriendCompleteDelegates(LocalUserNum, false, FriendId, ListName, FString(TEXT("DeleteFriend() is not supported")));
diff --git a/Engine/Plugins/Online/OnlineSubsystemSteam/Source/Private/OnlineFriendsInterfaceSteam.h b/Engine/Plugins/Online/OnlineSubsystemSteam/Source/Private/OnlineFriendsInterfaceSteam.h
--- a/Engine/Plugins/Online/OnlineSubsystemSteam/Source/Private/OnlineFriendsInterfaceSteam.h
+++ b/Engine/Plugins/Online/OnlineSubsystemSteam/Source/Private/OnlineFriendsInterfaceSteam.h
@@ -115,6 +115,7 @@ public:
 	virtual bool AcceptInvite(int32 LocalUserNum, const FUniqueNetId& FriendId, const FString& ListName, const FOnAcceptInviteComplete& Delegate = FOnAcceptInviteComplete()) override;
  	virtual bool RejectInvite(int32 LocalUserNum, const FUniqueNetId& FriendId, const FString& ListName) override;
 	virtual void SetFriendAlias(int32 LocalUserNum, const FUniqueNetId& FriendId, const FString& ListName, const FString& Alias, const FOnSetFriendAliasComplete& Delegate = FOnSetFriendAliasComplete()) override;
+	virtual void DeleteFriendAlias(int32 LocalUserNum, const FUniqueNetId& FriendId, const FString& ListName, const FOnDeleteFriendAliasComplete& Delegate = FOnDeleteFriendAliasComplete()) override;
 	virtual bool DeleteFriend(int32 LocalUserNum, const FUniqueNetId& FriendId, const FString& ListName) override;
 	virtual bool GetFriendsList(int32 LocalUserNum, const FString& ListName, TArray< TSharedRef<FOnlineFriend> >& OutFriends) override;
 	virtual TSharedPtr<FOnlineFriend> GetFriend(int32 LocalUserNum, const FUniqueNetId& FriendId, const FString& ListName) override;
diff --git a/Engine/Plugins/Online/OnlineSubsystemTencent/Source/Private/OnlineFriendsTencent.cpp b/Engine/Plugins/Online/OnlineSubsystemTencent/Source/Private/OnlineFriendsTencent.cpp
--- a/Engine/Plugins/Online/OnlineSubsystemTencent/Source/Private/OnlineFriendsTencent.cpp
+++ b/Engine/Plugins/Online/OnlineSubsystemTencent/Source/Private/OnlineFriendsTencent.cpp
@@ -479,6 +479,16 @@ void FOnlineFriendsTencent::SetFriendAlias(int32 LocalUserNum, const FUniqueNetI
 	});
 }
 
+void FOnlineFriendsTencent::DeleteFriendAlias(int32 LocalUserNum, const FUniqueNetId& FriendId, const FString& ListName, const FOnDeleteFriendAliasComplete& Delegate)
+{
+	TSharedRef<const FUniqueNetId> FriendIdRef = FriendId.AsShared();
+	TencentSubsystem->ExecuteNextTick([LocalUserNum, FriendIdRef, ListName, Delegate]()
+	{
+		UE_LOG_ONLINE_FRIEND(Warning, TEXT("FOnlineFriendsTencent::DeleteFriendAlias is not implemented"));
+		Delegate.ExecuteIfBound(LocalUserNum, *FriendIdRef, ListName, FOnlineError(EOnlineErrorResult::NotImplemented));
+	});
+}
+
 bool FOnlineFriendsTencent::DeleteFriend(int32 LocalUserNum, const FUniqueNetId& FriendId, const FString& ListName)
 {
 	/** NYI */
diff --git a/Engine/Plugins/Online/OnlineSubsystemTencent/Source/Private/OnlineFriendsTencent.h b/Engine/Plugins/Online/OnlineSubsystemTencent/Source/Private/OnlineFriendsTencent.h
--- a/Engine/Plugins/Online/OnlineSubsystemTencent/Source/Private/OnlineFriendsTencent.h
+++ b/Engine/Plugins/Online/OnlineSubsystemTencent/Source/Private/OnlineFriendsTencent.h
@@ -188,6 +188,7 @@ public:
 	virtual bool AcceptInvite(int32 LocalUserNum, const FUniqueNetId& FriendId, const FString& ListName, const FOnAcceptInviteComplete& Delegate = FOnAcceptInviteComplete()) override;
 	virtual bool RejectInvite(int32 LocalUserNum, const FUniqueNetId& FriendId, const FString& ListName) override;
 	virtual void SetFriendAlias(int32 LocalUserNum, const FUniqueNetId& FriendId, const FString& ListName, const FString& Alias, const FOnSetFriendAliasComplete& Delegate = FOnSetFriendAliasComplete()) override;
+	virtual void DeleteFriendAlias(int32 LocalUserNum, const FUniqueNetId& FriendId, const FString& ListName, const FOnDeleteFriendAliasComplete& Delegate = FOnDeleteFriendAliasComplete()) override;
 	virtual bool DeleteFriend(int32 LocalUserNum, const FUniqueNetId& FriendId, const FString& ListName) override;
 	virtual bool GetFriendsList(int32 LocalUserNum, const FString& ListName, TArray< TSharedRef<FOnlineFriend> >& OutFriends) override;
 	virtual TSharedPtr<FOnlineFriend> GetFriend(int32 LocalUserNum, const FUniqueNetId& FriendId, const FString& ListName) override;
