Added "LastOnline" field to OnlineUserPresence.

diff --git a/Engine/Plugins/Online/OnlineSubsystem/Source/Public/Interfaces/OnlinePresenceInterface.h b/Engine/Plugins/Online/OnlineSubsystem/Source/Public/Interfaces/OnlinePresenceInterface.h
--- a/Engine/Plugins/Online/OnlineSubsystem/Source/Public/Interfaces/OnlinePresenceInterface.h
+++ b/Engine/Plugins/Online/OnlineSubsystem/Source/Public/Interfaces/OnlinePresenceInterface.h
@@ -181,6 +181,7 @@ public:
 	uint32 bIsPlayingThisGame:1;
 	uint32 bIsJoinable:1;
 	uint32 bHasVoiceSupport:1;
+	FDateTime LastOnline;
 	FOnlineUserPresenceStatus Status;
 
 	/** Constructor */
@@ -198,6 +199,7 @@ public:
 		bIsJoinable = 0;
 		bHasVoiceSupport = 0;
 		Status = FOnlineUserPresenceStatus();
+		LastOnline = FDateTime::MaxValue();
 	}
 
 	const FString GetPlatform() const
