Add a configuration option to specify which module to load for a given OSS

diff --git a/Engine/Plugins/Online/OnlineSubsystem/Source/Private/OnlineSubsystemModule.cpp b/Engine/Plugins/Online/OnlineSubsystem/Source/Private/OnlineSubsystemModule.cpp
--- a/Engine/Plugins/Online/OnlineSubsystem/Source/Private/OnlineSubsystemModule.cpp
+++ b/Engine/Plugins/Online/OnlineSubsystem/Source/Private/OnlineSubsystemModule.cpp
@@ -15,15 +15,22 @@ IMPLEMENT_MODULE( FOnlineSubsystemModule, OnlineSubsystem );
  */
 FOnlineSubsystemDelegates::FOnOnlineSubsystemCreated FOnlineSubsystemDelegates::OnOnlineSubsystemCreated;
 
-/** Helper function to turn the friendly subsystem name into the module name */
-static inline FName GetOnlineModuleName(const FString& SubsystemName)
+static inline FName GetOnlineModuleName(const FString& SubsystemName, const TMap<FString, FName>& ModuleRedirects)
 {
 	FString ModuleBase(TEXT("OnlineSubsystem"));
 
 	FName ModuleName;
 	if (!SubsystemName.StartsWith(ModuleBase, ESearchCase::CaseSensitive))
 	{
-		ModuleName = FName(*(ModuleBase + SubsystemName));
+		const FName* const ModuleOverride = ModuleRedirects.Find(SubsystemName);
+		if (ModuleOverride != nullptr)
+		{
+			ModuleName = *ModuleOverride;
+		}
+		else
+		{
+			ModuleName = FName(*(ModuleBase + SubsystemName));
+		}
 	}
 	else
 	{
@@ -39,12 +46,12 @@ static inline FName GetOnlineModuleName(const FString& SubsystemName)
  * @param SubsystemName Name of the requested platform service to load
  * @return The module interface of the requested platform service, NULL if the service doesn't exist
  */
-static IModuleInterface* LoadSubsystemModule(const FString& SubsystemName)
+static IModuleInterface* LoadSubsystemModule(const FString& SubsystemName, const TMap<FString, FName>& ModuleRedirects)
 {
 	const bool bAttemptLoadModule = IOnlineSubsystem::IsEnabled(FName(*SubsystemName));
 	if (bAttemptLoadModule)
 	{
-		const FName ModuleName = GetOnlineModuleName(SubsystemName);
+		const FName ModuleName = GetOnlineModuleName(SubsystemName, ModuleRedirects);
 		FModuleManager& ModuleManager = FModuleManager::Get();
 
 		if (!ModuleManager.IsModuleLoaded(ModuleName))
@@ -73,6 +80,8 @@ void FOnlineSubsystemModule::StartupModule()
 		FModuleManager::Get().LoadModule(TEXT("XMPP"));
 	}
 
+	ProcessConfigDefinedModuleRedirects();
+
 	LoadDefaultSubsystem();
 	
 	// Also load the console/platform specific OSS which might not necessarily be the default OSS instance
@@ -95,15 +104,14 @@ void FOnlineSubsystemModule::ShutdownModule()
 	ShutdownOnlineSubsystem();
 }
 
-void FOnlineSubsystemModule::ProcessConfigDefinedSubsystems()
+static inline void ReadOnlineSubsystemConfigPairs(const TCHAR* Section, const TCHAR* Key, const FString& ConfigFile, TArray<TPair<FString, FString>>& OutPairs)
 {
-	// Save off the names of the user defined platform services.
-	TArray<FString> TmpConfigDefinedSubsystems;
-	GConfig->GetArray(TEXT("OnlineSubsystem"), TEXT("ConfigDefinedPlatformServices"), TmpConfigDefinedSubsystems, GEngineIni);
+	TArray<FString> ConfigPairs;
+	GConfig->GetArray(Section, Key, ConfigPairs, ConfigFile);
+	OutPairs.Reserve(ConfigPairs.Num());
 
-	// Takes on the pattern "(ServiceNameString=SubsystemName)"
-	// For example "(GameFeature=NULL)" to have OnlineSubsystemNull be the provider for "GameFeature"
-	for (const FString& ConfigEntry : TmpConfigDefinedSubsystems)
+	// Takes on the pattern "(Key=Value)"
+	for (const FString& ConfigEntry : ConfigPairs)
 	{
 		FString TrimmedConfigEntry = ConfigEntry.TrimStartAndEnd();
 		FString KeyString;
@@ -122,8 +130,35 @@ void FOnlineSubsystemModule::ProcessConfigDefinedSubsystems()
 			KeyString.TrimStartAndEndInline();
 			ValueString.TrimStartAndEndInline();
 		}
-		UE_LOG(LogOnline, Verbose, TEXT("ConfigDefinedPlatformServices: Associating OnlineSubsystem %s with identifier %s"), *ValueString, *KeyString);
-		ConfigDefinedSubsystems.Add(KeyString, FName(*ValueString));
+		OutPairs.Emplace(MoveTemp(KeyString), MoveTemp(ValueString));
+	}
+}
+
+void FOnlineSubsystemModule::ProcessConfigDefinedSubsystems()
+{
+	// Takes on the pattern "(ServiceNameString=SubsystemName)"
+	// For example "(GameFeature=NULL)" to have OnlineSubsystemNull be the provider for "GameFeature"
+	TArray<TPair<FString, FString>> ServicePairs;
+	ReadOnlineSubsystemConfigPairs(TEXT("OnlineSubsystem"), TEXT("ConfigDefinedPlatformServices"), GEngineIni, ServicePairs);
+
+	for (TPair<FString, FString>& ServicePair : ServicePairs)
+	{
+		UE_LOG(LogOnline, Verbose, TEXT("ConfigDefinedPlatformServices: Associating OnlineSubsystem %s with identifier %s"), *ServicePair.Value, *ServicePair.Key);
+		ConfigDefinedSubsystems.Add(MoveTemp(ServicePair.Key), FName(*ServicePair.Value));
+	}
+}
+
+void FOnlineSubsystemModule::ProcessConfigDefinedModuleRedirects()
+{
+	// Takes on the pattern "(SubsystemName=ModuleName)"
+	// For example "(Test=OnlineSubsystemNull)" to have OnlineSubsystemNull be the provider for "Test"
+	TArray<TPair<FString, FString>> RedirectPairs;
+	ReadOnlineSubsystemConfigPairs(TEXT("OnlineSubsystem"), TEXT("ModuleRedirects"), GEngineIni, RedirectPairs);
+
+	for (TPair<FString, FString>& RedirectPair : RedirectPairs)
+	{
+		UE_LOG(LogOnline, Verbose, TEXT("ProcessConfigDefinedModuleRedirects: Associating module %s with OnlineSubsystem %s"), *RedirectPair.Value, *RedirectPair.Key);
+		ModuleRedirects.Emplace(MoveTemp(RedirectPair.Key), FName(*RedirectPair.Value));
 	}
 }
 
@@ -132,7 +167,7 @@ bool FOnlineSubsystemModule::TryLoadSubsystemAndSetDefault(FName ModuleName)
 	// A module loaded with its factory method set for creation and a default instance of the online subsystem is required
 	bool bLoaded = false;
 	FString ModuleNameString = ModuleName.ToString();
-	if (LoadSubsystemModule(ModuleNameString))
+	if (LoadSubsystemModule(ModuleNameString, ModuleRedirects))
 	{
 		if (OnlineFactories.Contains(ModuleName))
 		{
@@ -219,7 +254,7 @@ void FOnlineSubsystemModule::ShutdownOnlineSubsystem()
 		UE_LOG_ONLINE(Verbose, TEXT("Unloading online subsystem: %s"), *It.Key().ToString());
 
 		// Unloading the module will do proper cleanup
-		FName ModuleName = GetOnlineModuleName(It.Key().ToString());
+		FName ModuleName = GetOnlineModuleName(It.Key().ToString(), ModuleRedirects);
 
 		const bool bIsShutdown = true;
 		ModuleManager.UnloadModule(ModuleName, bIsShutdown);
@@ -320,7 +355,7 @@ IOnlineSubsystem* FOnlineSubsystemModule::GetOnlineSubsystem(const FName InSubsy
 				if (OSSFactory == nullptr)
 				{
 					// Attempt to load the requested factory
-					IModuleInterface* NewModule = LoadSubsystemModule(SubsystemName.ToString());
+					IModuleInterface* NewModule = LoadSubsystemModule(SubsystemName.ToString(), ModuleRedirects);
 					if (NewModule)
 					{
 						// If the module loaded successfully this should be non-NULL
@@ -428,7 +463,7 @@ bool FOnlineSubsystemModule::IsOnlineSubsystemLoaded(const FName InSubsystemName
 
 	if (!SubsystemName.IsNone())
 	{
-		if (FModuleManager::Get().IsModuleLoaded(GetOnlineModuleName(SubsystemName.ToString())))
+		if (FModuleManager::Get().IsModuleLoaded(GetOnlineModuleName(SubsystemName.ToString(), ModuleRedirects)))
 		{
 			bIsLoaded = true;
 		}
diff --git a/Engine/Plugins/Online/OnlineSubsystem/Source/Public/OnlineSubsystemModule.h b/Engine/Plugins/Online/OnlineSubsystem/Source/Public/OnlineSubsystemModule.h
--- a/Engine/Plugins/Online/OnlineSubsystem/Source/Public/OnlineSubsystemModule.h
+++ b/Engine/Plugins/Online/OnlineSubsystem/Source/Public/OnlineSubsystemModule.h
@@ -50,6 +50,9 @@ private:
 	/** Have we warned already for a given online subsystem creation failure */
 	TMap<FName, bool> OnlineSubsystemFailureNotes;
 
+	/** Config driven override of module name for online subsystem */
+	TMap<FString, FName> ModuleRedirects;
+
 	/**
 	 * Transform an online subsystem identifier into its Subsystem and Instance constituents
 	 *
@@ -69,10 +72,15 @@ private:
 	FName ParseOnlineSubsystemName(const FName& FullName, FName& SubsystemName, FName& InstanceName) const;
 
 	/**
-	* Read any config defined subsystems from the configuration file
-	*/
+	 * Read any config defined subsystems from the configuration file
+	 */
 	void ProcessConfigDefinedSubsystems();
 
+	/**
+	 * Read any config defined subsystem module overrides from the configuration file
+	 */
+	void ProcessConfigDefinedModuleRedirects();
+
 	/**
 	 * Attempt to load the default subsystem specified in the configuration file
 	 */
@@ -121,13 +129,13 @@ public:
 	virtual IOnlineSubsystem* GetNativeSubsystem(bool bAutoLoad);
 
 	/**
-	* Get the online subsystem associated with the given config string
-	*
-	* @param ConfigString - Key to query for
-	* @param bAutoLoad - load the module if not already loaded
-	*
-	* @return pointer to the appropriate online subsystem
-	*/
+	 * Get the online subsystem associated with the given config string
+	 *
+	 * @param ConfigString - Key to query for
+	 * @param bAutoLoad - load the module if not already loaded
+	 *
+	 * @return pointer to the appropriate online subsystem
+	 */
 	virtual IOnlineSubsystem* GetSubsystemByConfig(const FString& ConfigString, bool bAutoLoad);
 
 	/**
