Remove another conflict marker
#rnx

diff --git a/Engine/Plugins/Online/OnlineSubsystem/Source/Public/Interfaces/VoiceInterface.h b/Engine/Plugins/Online/OnlineSubsystem/Source/Public/Interfaces/VoiceInterface.h
--- a/Engine/Plugins/Online/OnlineSubsystem/Source/Public/Interfaces/VoiceInterface.h
+++ b/Engine/Plugins/Online/OnlineSubsystem/Source/Public/Interfaces/VoiceInterface.h
@@ -573,7 +573,6 @@ public:
 	virtual void DisconnectAllEndpoints()
 	{
 	}
-<<<<
 };
 
 typedef TSharedPtr<IOnlineVoice, ESPMode::ThreadSafe> IOnlineVoicePtr;
