Get/Set(Subscription)Settings for OnlinePresenceMCP
Subscribe/Unsubscribe for OnlinePresenceMCP
BroadcastSubscriptionNotification for OnlinePresenceMCP
Get/Set Subscription Nudges for OnlinePresenceMCP (WIP)

diff --git a/Engine/Plugins/Online/OnlineSubsystem/Source/Public/OnlinePresenceErrors.h b/Engine/Plugins/Online/OnlineSubsystem/Source/Public/OnlinePresenceErrors.h
new file mode 100644
--- /dev/null
+++ b/Engine/Plugins/Online/OnlineSubsystem/Source/Public/OnlinePresenceErrors.h
@@ -0,0 +1,24 @@
+// Copyright 1998-2019 Epic Games, Inc. All Rights Reserved.
+
+#pragma once
+
+#include "OnlineError.h"
+#define LOCTEXT_NAMESPACE "OnlinePresence"
+#define ONLINE_ERROR_NAMESPACE "errors.com.epicgames.oss.presence"
+
+namespace OnlinePresence
+{
+	#include "OnlineErrorMacros.inl"
+
+	namespace Errors
+	{
+		inline FOnlineError InvalidResult() { return ONLINE_ERROR(EOnlineErrorResult::InvalidResults, TEXT("invalid_result")); }
+	}
+}
+
+#undef LOCTEXT_NAMESPACE
+#undef ONLINE_ERROR_NAMESPACE
+
+
+
+
