Add ECreatePartyCompletionResult error code for a non primary user trying to create a party

diff --git a/Engine/Plugins/Online/OnlineSubsystem/Source/Private/OnlinePartyInterface.cpp b/Engine/Plugins/Online/OnlineSubsystem/Source/Private/OnlinePartyInterface.cpp
--- a/Engine/Plugins/Online/OnlineSubsystem/Source/Private/OnlinePartyInterface.cpp
+++ b/Engine/Plugins/Online/OnlineSubsystem/Source/Private/OnlinePartyInterface.cpp
@@ -381,6 +381,10 @@ const TCHAR* ToString(const ECreatePartyCompletionResult Value)
 	{
 		return TEXT("LoggedOut");
 	}
+	case ECreatePartyCompletionResult::NotPrimaryUser:
+	{
+		return TEXT("NotPrimaryUser");
+	}
 	case ECreatePartyCompletionResult::UnknownInternalFailure:
 	{
 		return TEXT("UnknownInternalFailure");
diff --git a/Engine/Plugins/Online/OnlineSubsystem/Source/Public/Interfaces/OnlinePartyInterface.h b/Engine/Plugins/Online/OnlineSubsystem/Source/Public/Interfaces/OnlinePartyInterface.h
--- a/Engine/Plugins/Online/OnlineSubsystem/Source/Public/Interfaces/OnlinePartyInterface.h
+++ b/Engine/Plugins/Online/OnlineSubsystem/Source/Public/Interfaces/OnlinePartyInterface.h
@@ -1791,6 +1791,7 @@ enum class ECreatePartyCompletionResult : int8
 	FailedToCreateMucRoom,
 	NoResponse,
 	LoggedOut,
+	NotPrimaryUser,
 	UnknownInternalFailure = 0,
 	Succeeded = 1
 };
