OnlineIdentityInterface: Update FOnControllerPairingChanged to also provide the number of controllers now assigned to the New and Previous user in the event. If New or Previous is a null user, they will be considered to have 0 controllers. Updated comment to "[REDACTED_HANDLE] LocalUserNum..." to be more clear that LocalUserNum describes a person, not a device.

Anyone subscribing to this event will have to update the signature of their handler.

diff --git a/Engine/Plugins/Online/OnlineSubsystem/Source/Public/Interfaces/OnlineIdentityInterface.h b/Engine/Plugins/Online/OnlineSubsystem/Source/Public/Interfaces/OnlineIdentityInterface.h
--- a/Engine/Plugins/Online/OnlineSubsystem/Source/Public/Interfaces/OnlineIdentityInterface.h
+++ b/Engine/Plugins/Online/OnlineSubsystem/Source/Public/Interfaces/OnlineIdentityInterface.h
@@ -85,6 +85,31 @@ namespace EUserPrivileges
 	};
 }
 
+/**
+ * This struct will be broadcast with the FOnControllerPairingChanged event. This event is not broadcast on all OnlineIdentity platforms.
+ * There will be a struct for the previous user on the controller and for the new user on the controller.
+ */
+struct FControllerPairingChangedUserInfo
+{
+	/**
+	 * One of the users involved (the one who had the controller or the one who now has the controller).
+	 * This may be nobody (e.g. if the controller was assigned to a User, and now is assigned to nobody)
+	*/
+	const FUniqueNetId& User;
+
+	/**
+	 * The number of controllers now associated with User. If User is nobody, this number will be 0
+	*/
+	int32 ControllersRemaining;
+
+	FControllerPairingChangedUserInfo(const FUniqueNetId& InUser, int32 InControllersRemaining)
+		: User(InUser)
+		, ControllersRemaining(InControllersRemaining)
+	{
+	};
+
+};
+
 /**
  * Delegate called when a player logs in/out
  *
@@ -107,11 +132,11 @@ typedef FOnLoginStatusChanged::FDelegate FOnLoginStatusChangedDelegate;
 /**
  * Delegate called when a controller-user pairing changes
  *
- * @param LocalUserNum the controller number of the controller whose pairing changed
+ * @param LocalUserNum the logged-in user number of the user who owned the device whose pairing changed. If no previous owner, the user number of the new owner.
  * @param PreviousUser the user that used to be paired with this controller
  * @param NewUser the user that is currently paired with this controller
  */
-DECLARE_MULTICAST_DELEGATE_ThreeParams(FOnControllerPairingChanged, int /*LocalUserNum*/, const FUniqueNetId& /*PreviousUser*/, const FUniqueNetId& /*NewUser*/);
+DECLARE_MULTICAST_DELEGATE_ThreeParams(FOnControllerPairingChanged, int /*LocalUserNum*/, FControllerPairingChangedUserInfo /*PreviousUser*/, FControllerPairingChangedUserInfo /*NewUser*/);
 typedef FOnControllerPairingChanged::FDelegate FOnControllerPairingChangedDelegate;
 
 /**
@@ -216,7 +241,7 @@ public:
 	 * @param PreviousUser the previous user associated with the controller
 	 * @param NewUser the new user associated with the controller
 	 */
-	DEFINE_ONLINE_DELEGATE_THREE_PARAM(OnControllerPairingChanged, int /*LocalUserNum*/, const FUniqueNetId& /*PreviousUser*/, const FUniqueNetId& /*NewUser*/);
+	DEFINE_ONLINE_DELEGATE_THREE_PARAM(OnControllerPairingChanged, int /*LocalUserNum*/, FControllerPairingChangedUserInfo /*PreviousUser*/, FControllerPairingChangedUserInfo /*NewUser*/);
 
 	/**
 	 * Login/Authenticate with user credentials.
