Added OnlineTracing interface

diff --git a/Engine/Plugins/Online/OnlineSubsystem/Source/Private/OnlineTracingInterface.cpp b/Engine/Plugins/Online/OnlineSubsystem/Source/Private/OnlineTracingInterface.cpp
new file mode 100644
--- /dev/null
+++ b/Engine/Plugins/Online/OnlineSubsystem/Source/Private/OnlineTracingInterface.cpp
@@ -0,0 +1,43 @@
+// Copyright 1998-2019 Epic Games, Inc. All Rights Reserved.
+
+#include "Interfaces/OnlineTracingInterface.h"
+#include "Runtime/Core/Public/Features/IModularFeatures.h"
+
+namespace
+{
+	static FName GetModularFeatureName()
+	{
+		static FName FeatureName = FName(TEXT("OnlineTracing"));
+		return FeatureName;
+	}
+}
+
+/**
+ * Checks that the OnlineTracing feature is registered with IModularFeature
+ *
+ * @return FOnlineTracing singleton if feature registered, otherwise nullptr
+ */
+IOnlineTracing* IOnlineTracing::GetTracingHelper()
+{
+	if (IModularFeatures::Get().IsModularFeatureAvailable(GetModularFeatureName()))
+	{
+		return &IModularFeatures::Get().GetModularFeature<IOnlineTracing>(GetModularFeatureName());
+	}
+	return nullptr;
+}
+
+void IOnlineTracing::StartContext(FName ContextName)
+{
+	if (IOnlineTracing* OnlineTracingSingleton = GetTracingHelper())
+	{
+		OnlineTracingSingleton->StartContextImpl(ContextName);
+	}
+}
+
+void IOnlineTracing::EndContext(FName ContextName)
+{
+	if (IOnlineTracing* OnlineTracingSingleton = GetTracingHelper())
+	{
+		OnlineTracingSingleton->EndContextImpl(ContextName);
+	}
+}
\ No newline at end of file
diff --git a/Engine/Plugins/Online/OnlineSubsystem/Source/Public/Interfaces/OnlineTracingInterface.h b/Engine/Plugins/Online/OnlineSubsystem/Source/Public/Interfaces/OnlineTracingInterface.h
new file mode 100644
--- /dev/null
+++ b/Engine/Plugins/Online/OnlineSubsystem/Source/Public/Interfaces/OnlineTracingInterface.h
@@ -0,0 +1,42 @@
+// Copyright 1998-2019 Epic Games, Inc. All Rights Reserved.
+
+#pragma once
+
+#include "CoreMinimal.h"
+#include "Runtime/Core/Public/Features/IModularFeature.h"
+
+class FOutputDevice;
+
+/**
+ * Online Tracing Interface
+ */
+class ONLINESUBSYSTEM_API IOnlineTracing
+	: public IModularFeature
+{
+public:
+	/**
+	 * Starts a context
+	 *
+	 * @param ContextName user-specified context name.
+	 */
+
+	static void StartContext(FName ContextName);
+	
+	/**
+	 * Ends a context
+	 *
+	 * @param ContextName user-specified context name.
+	 */
+	static void EndContext(FName ContextName);
+
+	virtual bool GetUncompressedTracingLog(TArray<uint8>& OutLog) = 0;
+
+	virtual FString GetFilename() = 0;
+
+	static IOnlineTracing* GetTracingHelper();
+
+protected:
+	virtual void StartContextImpl(FName ContextName) = 0;
+
+	virtual void EndContextImpl(FName ContextName) = 0;
+};
\ No newline at end of file
