Returning first logged in user in identity interface, then first valid user if none found.

diff --git a/Engine/Plugins/Online/OnlineSubsystem/Source/Private/OnlineSubsystem.cpp b/Engine/Plugins/Online/OnlineSubsystem/Source/Private/OnlineSubsystem.cpp
--- a/Engine/Plugins/Online/OnlineSubsystem/Source/Private/OnlineSubsystem.cpp
+++ b/Engine/Plugins/Online/OnlineSubsystem/Source/Private/OnlineSubsystem.cpp
@@ -80,20 +80,28 @@ namespace OnlineIdentity
 /** Workaround, please avoid using this */
 TSharedPtr<const FUniqueNetId> GetFirstSignedInUser(IOnlineIdentityPtr IdentityInt)
 {
-	TSharedPtr<const FUniqueNetId> UserId = nullptr;
 	if (IdentityInt.IsValid())
 	{
+		// find an entry for a fully logged in user
 		for (int32 i = 0; i < MAX_LOCAL_PLAYERS; i++)
 		{
-			UserId = IdentityInt->GetUniquePlayerId(i);
+			TSharedPtr<const FUniqueNetId> UserId = IdentityInt->GetUniquePlayerId(i);
+			if (UserId.IsValid() && UserId->IsValid() && IdentityInt->GetLoginStatus(*UserId) == ELoginStatus::LoggedIn)
+			{
+				return UserId;
+			}
+		}
+		// find an entry for a locally logged in user
+		for (int32 i = 0; i < MAX_LOCAL_PLAYERS; i++)
+		{
+			TSharedPtr<const FUniqueNetId> UserId = IdentityInt->GetUniquePlayerId(i);
 			if (UserId.IsValid() && UserId->IsValid())
 			{
-				break;
+				return UserId;
 			}
 		}
 	}
-
-	return UserId;
+	return nullptr;
 }
 
 int32 GetBuildUniqueId()
