#online,title

Allow for multiple overlapping title file enumerations as long as the Page matches

diff --git a/Engine/Plugins/Online/OnlineSubsystem/Source/Public/OnlineSubsystemTypes.h b/Engine/Plugins/Online/OnlineSubsystem/Source/Public/OnlineSubsystemTypes.h
--- a/Engine/Plugins/Online/OnlineSubsystem/Source/Public/OnlineSubsystemTypes.h
+++ b/Engine/Plugins/Online/OnlineSubsystem/Source/Public/OnlineSubsystemTypes.h
@@ -1043,6 +1043,11 @@ public:
 		return Start >= 0 && Count >= 0;
 	}
 
+	bool operator==(const FPagedQuery& Other) const
+	{
+		return Other.Start == Start && Other.Count == Count;
+	}
+
 	/** first entry to fetch */
 	int32 Start;
 	/** total entries to fetch. -1 means ALL */
