Adds most of the requested logging. Records sent time in additton to received time.

diff --git a/Engine/Plugins/Online/OnlineSubsystem/Source/Private/OnlinePartyInterface.cpp b/Engine/Plugins/Online/OnlineSubsystem/Source/Private/OnlinePartyInterface.cpp
--- a/Engine/Plugins/Online/OnlineSubsystem/Source/Private/OnlinePartyInterface.cpp
+++ b/Engine/Plugins/Online/OnlineSubsystem/Source/Private/OnlinePartyInterface.cpp
@@ -866,15 +866,7 @@ FString ToDebugString(const FPartyConfiguration& PartyConfiguration)
 
 FString ToDebugString(const IOnlinePartyJoinInfo& JoinInfo)
 {
-	return FString::Printf(TEXT("SourceUserId(%s) SourceDisplayName(%s) PartyId(%s) HasKey(%d) HasPassword(%d) IsAcceptingMembers(%d) NotAcceptingReason(%d)"),
-			*(JoinInfo.GetSourceUserId()->ToDebugString()),
-			*(JoinInfo.GetSourceDisplayName()),
-			*(JoinInfo.GetPartyId()->ToDebugString()),
-			JoinInfo.HasKey() ? 1 : 0,
-			JoinInfo.HasPassword() ? 1 : 0,
-			JoinInfo.IsAcceptingMembers() ? 1 : 0,
-			JoinInfo.GetNotAcceptingReason()
-		);
+	return JoinInfo.ToDebugString();
 }
 
 /**
@@ -914,3 +906,14 @@ FString ToDebugString(const FOnlinePartyData& PartyData)
 	Result += ToDebugString(PartyData.GetKeyValAttrs());
 	return Result;
 }
+
+FString IOnlinePartyJoinInfo::ToDebugString() const 
+{
+	return FString::Printf(TEXT("SourceUserId(%s) SourceDisplayName(%s) PartyId(%s) HasKey(%d) HasPassword(%d) IsAcceptingMembers(%d) NotAcceptingReason(%d)"),
+		*(GetSourceUserId()->ToDebugString()),
+		*(GetSourceDisplayName()),
+		*(GetPartyId()->ToDebugString()),
+		HasKey() ? 1 : 0, HasPassword() ? 1 : 0,
+		IsAcceptingMembers() ? 1 : 0,
+		GetNotAcceptingReason());
+};
\ No newline at end of file
diff --git a/Engine/Plugins/Online/OnlineSubsystem/Source/Public/Interfaces/OnlinePartyInterface.h b/Engine/Plugins/Online/OnlineSubsystem/Source/Public/Interfaces/OnlinePartyInterface.h
--- a/Engine/Plugins/Online/OnlineSubsystem/Source/Public/Interfaces/OnlinePartyInterface.h
+++ b/Engine/Plugins/Online/OnlineSubsystem/Source/Public/Interfaces/OnlinePartyInterface.h
@@ -524,6 +524,9 @@ public:
 	 * @return whether or not the join info has the info to request an invite
 	 */
 	virtual bool CanRequestAnInvite() const = 0;
+
+	ONLINESUBSYSTEM_API virtual FString ToDebugString() const;
+
 };
 
 typedef TSharedRef<const IOnlinePartyJoinInfo> IOnlinePartyJoinInfoConstRef;
diff --git a/Engine/Source/Runtime/Online/XMPP/Private/XmppStrophe/XmppPresenceStrophe.cpp b/Engine/Source/Runtime/Online/XMPP/Private/XmppStrophe/XmppPresenceStrophe.cpp
--- a/Engine/Source/Runtime/Online/XMPP/Private/XmppStrophe/XmppPresenceStrophe.cpp
+++ b/Engine/Source/Runtime/Online/XMPP/Private/XmppStrophe/XmppPresenceStrophe.cpp
@@ -114,6 +114,8 @@ bool FXmppPresenceStrophe::ReceiveStanza(const FStropheStanza& IncomingStanza)
 			FDateTime::ParseIso8601(*TimestampStanza->GetAttribute(Strophe::SA_STAMP), Presence.SentTime);
 		}
 
+		Presence.ReceivedTime = FDateTime::UtcNow();
+
 		FString UnusedPlatformUserId;
 		Presence.UserJid.ParseResource(Presence.AppId, Presence.Platform, UnusedPlatformUserId);
 	}
diff --git a/Engine/Source/Runtime/Online/XMPP/Public/XmppPresence.h b/Engine/Source/Runtime/Online/XMPP/Public/XmppPresence.h
--- a/Engine/Source/Runtime/Online/XMPP/Public/XmppPresence.h
+++ b/Engine/Source/Runtime/Online/XMPP/Public/XmppPresence.h
@@ -66,6 +66,8 @@ public:
 	bool bIsAvailable;
 	/** time when presence was sent by the user */
 	FDateTime SentTime;
+	/** time when presence was sent by the user */
+	FDateTime ReceivedTime;
 	/** app id user is logged in from */
 	FString AppId;
 	/** platform associated with this client */
