Attempt to fix Mac CIS failure
#rnx

diff --git a/Engine/Plugins/Online/OnlineSubsystem/Source/Public/Interfaces/OnlinePartyInterface.h b/Engine/Plugins/Online/OnlineSubsystem/Source/Public/Interfaces/OnlinePartyInterface.h
--- a/Engine/Plugins/Online/OnlineSubsystem/Source/Public/Interfaces/OnlinePartyInterface.h
+++ b/Engine/Plugins/Online/OnlineSubsystem/Source/Public/Interfaces/OnlinePartyInterface.h
@@ -441,7 +441,7 @@ typedef TSharedPtr<const IOnlinePartyPendingJoinRequestInfo> IOnlinePartyPending
 /**
  * Info needed to join a party
  */
-class IOnlinePartyJoinInfo
+class ONLINESUBSYSTEM_API IOnlinePartyJoinInfo
 	: public TSharedFromThis<IOnlinePartyJoinInfo>
 {
 public:
@@ -525,7 +525,7 @@ public:
 	 */
 	virtual bool CanRequestAnInvite() const = 0;
 
-	ONLINESUBSYSTEM_API virtual FString ToDebugString() const;
+	virtual FString ToDebugString() const;
 
 };
 
