remove circular dependency in OnlineSubsystem

diff --git a/Engine/Plugins/Online/OnlineSubsystem/Source/OnlineSubsystem.Build.cs b/Engine/Plugins/Online/OnlineSubsystem/Source/OnlineSubsystem.Build.cs
--- a/Engine/Plugins/Online/OnlineSubsystem/Source/OnlineSubsystem.Build.cs
+++ b/Engine/Plugins/Online/OnlineSubsystem/Source/OnlineSubsystem.Build.cs
@@ -10,8 +10,7 @@ public class OnlineSubsystem : ModuleRules
 		
 		PublicDependencyModuleNames.AddRange(
 			new string[] {
-				"Json",
-                "OnlineSubsystemUtils"
+				"Json"
             }
 		);
 
diff --git a/Engine/Plugins/Online/OnlineSubsystem/Source/Public/Interfaces/OnlineStoreInterface.h b/Engine/Plugins/Online/OnlineSubsystem/Source/Public/Interfaces/OnlineStoreInterface.h
--- a/Engine/Plugins/Online/OnlineSubsystem/Source/Public/Interfaces/OnlineStoreInterface.h
+++ b/Engine/Plugins/Online/OnlineSubsystem/Source/Public/Interfaces/OnlineStoreInterface.h
@@ -6,7 +6,7 @@
 #include "UObject/ObjectMacros.h"
 #include "OnlineSubsystemTypes.h"
 #include "OnlineDelegateMacros.h"
-#include "InAppPurchaseCallbackProxy2.h"
+//#include "InAppPurchaseCallbackProxy2.h"
 #include "OnlineStoreInterface.generated.h"
 
 ONLINESUBSYSTEM_API DECLARE_LOG_CATEGORY_EXTERN(LogOnlineStore, Log, All);
