Adding retries to lan queries

diff --git a/Engine/Plugins/Online/OnlineSubsystem/Source/Private/LANBeacon.cpp b/Engine/Plugins/Online/OnlineSubsystem/Source/Private/LANBeacon.cpp
--- a/Engine/Plugins/Online/OnlineSubsystem/Source/Private/LANBeacon.cpp
+++ b/Engine/Plugins/Online/OnlineSubsystem/Source/Private/LANBeacon.cpp
@@ -132,6 +132,8 @@ bool FLANSession::Host(FOnValidQueryPacketDelegate& QueryDelegate)
 		StopLANSession();
 	}
 
+	LanBeaconGuid = FGuid::NewGuid();
+
 	// Bind a socket for LAN beacon activity
 	LanBeacon = new FLanBeacon();
 	if (LanBeacon->Init(LanAnnouncePort))
@@ -163,6 +165,10 @@ bool FLANSession::Search(FNboSerializeToBuffer& Packet, FOnValidResponsePacketDe
 	{
 		StopLANSession();
 	}
+	
+	LanBeaconGuid = FGuid::NewGuid();
+
+	CachedResponseGuids.Empty();
 
 	// Bind a socket for LAN beacon activity
 	LanBeacon = new FLanBeacon();
@@ -184,6 +190,11 @@ bool FLANSession::Search(FNboSerializeToBuffer& Packet, FOnValidResponsePacketDe
 			// Set the timestamp for timing out a search
 			LanQueryTimeLeft = LanQueryTimeout;
 
+			RetryData.Reset();
+			RetryData.Append(Packet.GetBuffer().GetData(), Packet.GetByteCount());
+
+			LanQueryRetryTimeLeft = LanQueryRetryTime;
+
 			AddOnValidResponsePacketDelegate_Handle(ResponseDelegate);
 			AddOnSearchingTimeoutDelegate_Handle(TimeoutDelegate);
 		}
@@ -214,6 +225,10 @@ void FLANSession::StopLANSession()
 	OnValidQueryPacketDelegates.Clear();
 	OnValidResponsePacketDelegates.Clear();
 	OnSearchingTimeoutDelegates.Clear();
+
+	RetryData.Empty();
+
+	CachedResponseGuids.Empty();
 }
 
 void FLANSession::Tick(float DeltaTime)
@@ -244,11 +259,19 @@ void FLANSession::Tick(float DeltaTime)
 			}
 			else if (LanBeaconState == ELanBeaconState::Searching)
 			{
+				FGuid ResponseGuid;
+
 				// We can only accept Server Response packets
-				if (IsValidLanResponsePacket(PacketData, NumRead))
+				if (IsValidLanResponsePacket(PacketData, NumRead, ResponseGuid))
 				{
-					// Strip off the header
-					TriggerOnValidResponsePacketDelegates(&PacketData[LAN_BEACON_PACKET_HEADER_SIZE], NumRead - LAN_BEACON_PACKET_HEADER_SIZE);
+					bool bAlreadySet = false;
+					CachedResponseGuids.Add(ResponseGuid, &bAlreadySet);
+
+					if (!bAlreadySet)
+					{
+						// Strip off the header
+						TriggerOnValidResponsePacketDelegates(&PacketData[LAN_BEACON_PACKET_HEADER_SIZE], NumRead - LAN_BEACON_PACKET_HEADER_SIZE);
+					}
 				}
 			}
 		}
@@ -256,6 +279,22 @@ void FLANSession::Tick(float DeltaTime)
 		{
 			if (LanBeaconState == ELanBeaconState::Searching)
 			{
+				LanQueryRetryTimeLeft -= DeltaTime;
+				
+				if (LanBeacon && (LanQueryRetryTimeLeft <= 0.f) && (LanQueryTimeLeft > LanQueryRetryTime))
+				{
+					if (LanBeacon->BroadcastPacket(RetryData.GetData(), RetryData.Num()))
+					{
+						UE_LOG_ONLINE(Verbose, TEXT("Sent retry query packet..."));
+					}
+					else
+					{
+						UE_LOG_ONLINE(Warning, TEXT("Failed to send retry discovery broadcast %s"), ISocketSubsystem::Get(PLATFORM_SOCKETSUBSYSTEM)->GetSocketError());
+					}
+
+					LanQueryRetryTimeLeft = LanQueryRetryTime;
+				}
+
 				// Decrement the amount of time remaining
 				LanQueryTimeLeft -= DeltaTime;
 				// Check for a timeout on the search packet
@@ -280,7 +319,9 @@ void FLANSession::CreateHostResponsePacket(FNboSerializeToBuffer& Packet, uint64
 		// Add the packet type
 		<< LAN_SERVER_RESPONSE1 << LAN_SERVER_RESPONSE2
 		// Append the client nonce as a uint64
-		<< ClientNonce;
+		<< ClientNonce
+		// Append the beacon guid
+		<< LanBeaconGuid;
 }
 
 void FLANSession::CreateClientQueryPacket(FNboSerializeToBuffer& Packet, uint64 ClientNonce)
@@ -294,7 +335,9 @@ void FLANSession::CreateClientQueryPacket(FNboSerializeToBuffer& Packet, uint64
 		// Identify the packet type
 		<< LAN_SERVER_QUERY1 << LAN_SERVER_QUERY2
 		// Append the nonce as a uint64
-		<< ClientNonce;
+		<< ClientNonce
+		// Append the beacon guid
+		<< LanBeaconGuid;
 }
 
 /**
@@ -358,6 +401,10 @@ bool FLANSession::IsValidLanQueryPacket(const uint8* Packet, uint32 Length, uint
 					bIsValid = (SQ1 == LAN_SERVER_QUERY1 && SQ2 == LAN_SERVER_QUERY2);
 					// Read the client nonce as the outvalue
 					PacketReader >> ClientNonce;
+
+					// Read the client's beacon guid 
+					FGuid BeaconGuid;
+					PacketReader >> BeaconGuid;
 				}
 			}
 		}
@@ -370,10 +417,11 @@ bool FLANSession::IsValidLanQueryPacket(const uint8* Packet, uint32 Length, uint
  *
  * @param Packet the packet data to check
  * @param Length the size of the packet buffer
+ * @param ResponseGuid the lan beacon guid in the response packet
  *
  * @return true if the header is valid, false otherwise
  */
-bool FLANSession::IsValidLanResponsePacket(const uint8* Packet, uint32 Length)
+bool FLANSession::IsValidLanResponsePacket(const uint8* Packet, uint32 Length, FGuid& ResponseGuid)
 {
 	bool bIsValid = false;
 	// Serialize out the data if the packet is the right size
@@ -404,7 +452,19 @@ bool FLANSession::IsValidLanResponsePacket(const uint8* Packet, uint32 Length)
 					{
 						uint64 Nonce = 0;
 						PacketReader >> Nonce;
-						bIsValid = (Nonce == LanNonce);
+						
+						if (Nonce == LanNonce)
+						{
+							FGuid BeaconGuid;
+							PacketReader >> BeaconGuid;
+
+							bIsValid = !PacketReader.HasOverflow();
+
+							if (bIsValid)
+							{
+								ResponseGuid = BeaconGuid;
+							}
+						}
 					}
 				}
 			}
diff --git a/Engine/Plugins/Online/OnlineSubsystem/Source/Public/LANBeacon.h b/Engine/Plugins/Online/OnlineSubsystem/Source/Public/LANBeacon.h
--- a/Engine/Plugins/Online/OnlineSubsystem/Source/Public/LANBeacon.h
+++ b/Engine/Plugins/Online/OnlineSubsystem/Source/Public/LANBeacon.h
@@ -6,6 +6,7 @@
 #include "OnlineSubsystemTypes.h"
 #include "OnlineDelegateMacros.h"
 #include "Misc/ConfigCacheIni.h"
+#include "Misc/Guid.h"
 
 /**
  * This value indicates which packet version the server is sending. Clients with
@@ -13,12 +14,22 @@
  * changing the packet format and there are existing servers on the network
  * Current format:
  *
- *	<Ver byte><Platform byte><Game unique 4 bytes><packet type 2 bytes><nonce 8 bytes><payload>
+ *	<Ver byte><Platform byte><Game unique 4 bytes><packet type 2 bytes><nonce 8 bytes><guid 16 bytes><payload>
  */
-#define LAN_BEACON_PACKET_VERSION (uint8)11
+enum class ELANBeaconVersionHistory : uint8
+{
+	HISTORY_INITIAL				= 11,
+	HISTORY_BEACON_GUID			= 12,	// Lan response packets now contain an instance guid
+	
+	// -----<new versions can be added before this line>-------------------------------------------------
+	HISTORY_PLUS_ONE,
+	HISTORY_LATEST 							= HISTORY_PLUS_ONE - 1
+};
+
+static const uint8 LAN_BEACON_PACKET_VERSION = (uint8)ELANBeaconVersionHistory::HISTORY_LATEST;
 
 /** The size of the header for validation */
-#define LAN_BEACON_PACKET_HEADER_SIZE 16
+#define LAN_BEACON_PACKET_HEADER_SIZE 32
 	
 // Offsets for various fields
 #define LAN_BEACON_VER_OFFSET 0
@@ -27,6 +38,7 @@
 #define LAN_BEACON_PACKETTYPE1_OFFSET 6
 #define LAN_BEACON_PACKETTYPE2_OFFSET 7
 #define LAN_BEACON_NONCE_OFFSET 8
+#define LAN_BEACON_GUID_OFFSET 16
 
 // Packet types in 2 byte readable form
 #define LAN_SERVER_QUERY1 (uint8)'S'
@@ -106,6 +118,7 @@ public:
 #define LAN_ANNOUNCE_PORT 14001
 #define LAN_UNIQUE_ID 9999
 #define LAN_QUERY_TIMEOUT 5
+#define LAN_QUERY_RETRY_TIME 1
 #define LAN_PLATFORMMASK 0xffffffff
 
 /**
@@ -132,7 +145,23 @@ protected:
 	 *
 	 * @return true if the header is valid, false otherwise
 	 */
-	bool IsValidLanResponsePacket(const uint8* Packet, uint32 Length);
+	UE_DEPRECATED(4.24, "IsValidLanResponsePacket now takes an additional guid reference")
+	bool IsValidLanResponsePacket(const uint8* Packet, uint32 Length)
+	{
+		FGuid CompatGuid;
+		return IsValidLanResponsePacket(Packet, Length, CompatGuid);
+	}
+
+	/**
+	 * Determines if the packet header is valid or not
+	 *
+	 * @param Packet the packet data to check
+	 * @param Length the size of the packet buffer
+	 * @param ResponseGuid the beacon guid from the response packet
+	 *
+	 * @return true if the header is valid, false otherwise
+	 */
+	bool IsValidLanResponsePacket(const uint8* Packet, uint32 Length, FGuid& ResponseGuid);
 
 public:
 
@@ -148,6 +177,9 @@ public:
 	/** The amount of time to wait before timing out a LAN query request */
 	float LanQueryTimeout;
 
+	/** The amount of time to wait before resending a LAN query request */
+	float LanQueryRetryTime;
+
 	/** LAN beacon for packet broadcast */
 	class FLanBeacon* LanBeacon;
 
@@ -160,15 +192,29 @@ public:
 	/** The amount of time before the LAN query is considered done */
 	float LanQueryTimeLeft;
 
+	/** The amount of time before the LAN query retry */
+	float LanQueryRetryTimeLeft;
+
+	/** Cached search packet data for retries */
+	TArray<uint8> RetryData;
+
+	/** Unique identifier for this hosting/searching LAN session */
+	FGuid LanBeaconGuid;
+
+	/** Cached guids already received during current search */
+	TSet<FGuid> CachedResponseGuids;
+
 	FLANSession() :
 		LanAnnouncePort(LAN_ANNOUNCE_PORT),
 		LanGameUniqueId(LAN_UNIQUE_ID),
 		LanPacketPlatformMask(LAN_PLATFORMMASK),
 		LanQueryTimeout(LAN_QUERY_TIMEOUT),
-		LanBeacon(NULL),
+		LanQueryRetryTime(LAN_QUERY_RETRY_TIME),
+		LanBeacon(nullptr),
 		LanBeaconState(ELanBeaconState::NotUsingLanBeacon),
 		LanNonce(0),
-		LanQueryTimeLeft(0.0f)
+		LanQueryTimeLeft(0.0f),
+		LanQueryRetryTimeLeft(0.0f)
 	{
 		if (!GConfig->GetInt(TEXT("LANSession"), TEXT("LanAnnouncePort"), LanAnnouncePort, GEngineIni))
 		{
