Fix for warning after upgrading to PVS-Studio 7.0.4:

V1043: A global object variable 'DefaultPresenceKey' is declared in the header. Multiple copies of it will be created in all translation units that include this header file.

#rnx

diff --git a/Engine/Plugins/Online/OnlineSubsystem/Source/Private/OnlineSubsystem.cpp b/Engine/Plugins/Online/OnlineSubsystem/Source/Private/OnlineSubsystem.cpp
--- a/Engine/Plugins/Online/OnlineSubsystem/Source/Private/OnlineSubsystem.cpp
+++ b/Engine/Plugins/Online/OnlineSubsystem/Source/Private/OnlineSubsystem.cpp
@@ -61,6 +61,27 @@ ONLINESUBSYSTEM_API DEFINE_STAT(STAT_Session_Interface);
 ONLINESUBSYSTEM_API DEFINE_STAT(STAT_Voice_Interface);
 #endif
 
+/** The default key that will update presence text in the platform's UI */
+const FString DefaultPresenceKey = TEXT("RichPresence");
+
+/** Custom presence data that is not seen by users but can be polled */
+const FString CustomPresenceDataKey = TEXT("CustomData");
+
+/** Name of the client that sent the presence update */
+const FString DefaultAppIdKey = TEXT("AppId");
+
+/** Platform of the client that sent the presence update */
+const FString DefaultPlatformKey = TEXT("Platform");
+
+/** Override Id of the client to set the presence state to */
+const FString OverrideAppIdKey = TEXT("OverrideAppId");
+
+/** Id of the session for the presence update. @todo samz - SessionId on presence data should be FUniqueNetId not uint64 */
+const FString DefaultSessionIdKey = TEXT("SessionId");
+
+/** Resource the client is logged in with */
+const FString PresenceResourceKey = TEXT("ResourceKey");
+
 namespace OnlineIdentity
 {
 	namespace Errors
diff --git a/Engine/Plugins/Online/OnlineSubsystem/Source/Public/Interfaces/OnlinePresenceInterface.h b/Engine/Plugins/Online/OnlineSubsystem/Source/Public/Interfaces/OnlinePresenceInterface.h
--- a/Engine/Plugins/Online/OnlineSubsystem/Source/Public/Interfaces/OnlinePresenceInterface.h
+++ b/Engine/Plugins/Online/OnlineSubsystem/Source/Public/Interfaces/OnlinePresenceInterface.h
@@ -27,25 +27,25 @@ typedef FString FPresenceKey;
 typedef FOnlineKeyValuePairs<FPresenceKey, FVariantData> FPresenceProperties;
 
 /** The default key that will update presence text in the platform's UI */
-const FString DefaultPresenceKey = TEXT("RichPresence");
+extern ONLINESUBSYSTEM_API const FString DefaultPresenceKey;
 
 /** Custom presence data that is not seen by users but can be polled */
-const FString CustomPresenceDataKey = TEXT("CustomData");
+extern ONLINESUBSYSTEM_API const FString CustomPresenceDataKey;
 
 /** Name of the client that sent the presence update */
-const FString DefaultAppIdKey = TEXT("AppId");
+extern ONLINESUBSYSTEM_API const FString DefaultAppIdKey;
 
 /** Platform of the client that sent the presence update */
-const FString DefaultPlatformKey = TEXT("Platform");
+extern ONLINESUBSYSTEM_API const FString DefaultPlatformKey;
 
 /** Override Id of the client to set the presence state to */
-const FString OverrideAppIdKey = TEXT("OverrideAppId");
+extern ONLINESUBSYSTEM_API const FString OverrideAppIdKey;
 
 /** Id of the session for the presence update. @todo samz - SessionId on presence data should be FUniqueNetId not uint64 */
-const FString DefaultSessionIdKey = TEXT("SessionId");
+extern ONLINESUBSYSTEM_API const FString DefaultSessionIdKey;
 
 /** Resource the client is logged in with */
-const FString PresenceResourceKey = TEXT("ResourceKey");
+extern ONLINESUBSYSTEM_API const FString PresenceResourceKey;
 
 namespace EOnlinePresenceState
 {
