- Include the AppName in presence updates

diff --git a/Engine/Plugins/Online/OnlineSubsystem/Source/Public/Interfaces/OnlinePresenceInterface.h b/Engine/Plugins/Online/OnlineSubsystem/Source/Public/Interfaces/OnlinePresenceInterface.h
--- a/Engine/Plugins/Online/OnlineSubsystem/Source/Public/Interfaces/OnlinePresenceInterface.h
+++ b/Engine/Plugins/Online/OnlineSubsystem/Source/Public/Interfaces/OnlinePresenceInterface.h
@@ -35,6 +35,9 @@ extern ONLINESUBSYSTEM_API const FString CustomPresenceDataKey;
 /** Name of the client that sent the presence update */
 extern ONLINESUBSYSTEM_API const FString DefaultAppIdKey;
 
+/** User friendly name of the client that sent the presence update */
+const FString DefaultProductNameKey = TEXT("ProductName");
+
 /** Platform of the client that sent the presence update */
 extern ONLINESUBSYSTEM_API const FString DefaultPlatformKey;
 
