Change OnlineSubsystem to only rely on the engine-less AudioMixerCore, since all it needs is a class that lives in the signal processing module.

diff --git a/Engine/Plugins/Online/OnlineSubsystem/Source/OnlineSubsystem.Build.cs b/Engine/Plugins/Online/OnlineSubsystem/Source/OnlineSubsystem.Build.cs
--- a/Engine/Plugins/Online/OnlineSubsystem/Source/OnlineSubsystem.Build.cs
+++ b/Engine/Plugins/Online/OnlineSubsystem/Source/OnlineSubsystem.Build.cs
@@ -26,7 +26,8 @@ public class OnlineSubsystem : ModuleRules
 				"ImageCore",
 				"Sockets",
 				"JsonUtilities",
-				"AudioMixer",
+				"AudioMixerCore",
+				"SignalProcessing",
 			}
 		);
 	}
