Make SignalProcessing a public dependency of OnlineSubsystem, so that subclasses of IVoiceInterface in modules that rely on OnlineSubsytem can include Dsp/MultithreadedPatching.h.

diff --git a/Engine/Plugins/Online/OnlineSubsystem/Source/OnlineSubsystem.Build.cs b/Engine/Plugins/Online/OnlineSubsystem/Source/OnlineSubsystem.Build.cs
--- a/Engine/Plugins/Online/OnlineSubsystem/Source/OnlineSubsystem.Build.cs
+++ b/Engine/Plugins/Online/OnlineSubsystem/Source/OnlineSubsystem.Build.cs
@@ -10,8 +10,9 @@ public class OnlineSubsystem : ModuleRules
 		
 		PublicDependencyModuleNames.AddRange(
 			new string[] {
-				"Json"
-            }
+				"Json",
+				"SignalProcessing"
+			}
 		);
 
 		PublicIncludePaths.Add(ModuleDirectory);
@@ -27,7 +28,6 @@ public class OnlineSubsystem : ModuleRules
 				"Sockets",
 				"JsonUtilities",
 				"AudioMixerCore",
-				"SignalProcessing",
 			}
 		);
 	}
