Add proactive comment that OnlineSubsystem cannot depend on the Engine module
#rnx
#tests compile win64 editor

diff --git a/Engine/Plugins/Online/OnlineSubsystem/Source/OnlineSubsystem.Build.cs b/Engine/Plugins/Online/OnlineSubsystem/Source/OnlineSubsystem.Build.cs
--- a/Engine/Plugins/Online/OnlineSubsystem/Source/OnlineSubsystem.Build.cs
+++ b/Engine/Plugins/Online/OnlineSubsystem/Source/OnlineSubsystem.Build.cs
@@ -20,6 +20,7 @@ public class OnlineSubsystem : ModuleRules
         PublicDefinitions.Add("ONLINESUBSYSTEM_PACKAGE=1");
 		PublicDefinitions.Add("DEBUG_LAN_BEACON=0");
 
+		// OnlineSubsystem cannot depend on Engine!
 		PrivateDependencyModuleNames.AddRange(
 			new string[] { 
 				"Core", 
