- Fix copy paste error in FOnlineAsyncTaskSwitchGetUserPrivilege
- Differentiate FOnlineAsyncTaskThreadedGenericCallable tasks in the logs

diff --git a/Engine/Plugins/Online/OnlineSubsystem/Source/Private/OnlineIdentityInterface.cpp b/Engine/Plugins/Online/OnlineSubsystem/Source/Private/OnlineIdentityInterface.cpp
new file mode 100644
--- /dev/null
+++ b/Engine/Plugins/Online/OnlineSubsystem/Source/Private/OnlineIdentityInterface.cpp
@@ -0,0 +1,42 @@
+// Copyright Epic Games, Inc. All Rights Reserved.
+
+#include "Interfaces/OnlineIdentityInterface.h"
+
+FString ToDebugString(IOnlineIdentity::EPrivilegeResults PrivilegeResults)
+{
+	if (PrivilegeResults == IOnlineIdentity::EPrivilegeResults::NoFailures)
+	{
+		return TEXT("NoFailures");
+	}
+	else
+	{
+		TArray<FString> ResultNames;
+		if ((uint32)PrivilegeResults & (uint32)IOnlineIdentity::EPrivilegeResults::RequiredPatchAvailable) ResultNames.Emplace(TEXT("RequiredPatchAvailable"));
+		if ((uint32)PrivilegeResults & (uint32)IOnlineIdentity::EPrivilegeResults::RequiredSystemUpdate) ResultNames.Emplace(TEXT("RequiredSystemUpdate"));
+		if ((uint32)PrivilegeResults & (uint32)IOnlineIdentity::EPrivilegeResults::AgeRestrictionFailure) ResultNames.Emplace(TEXT("AgeRestrictionFailure"));
+		if ((uint32)PrivilegeResults & (uint32)IOnlineIdentity::EPrivilegeResults::AccountTypeFailure) ResultNames.Emplace(TEXT("AccountTypeFailure"));
+		if ((uint32)PrivilegeResults & (uint32)IOnlineIdentity::EPrivilegeResults::UserNotFound) ResultNames.Emplace(TEXT("UserNotFound"));
+		if ((uint32)PrivilegeResults & (uint32)IOnlineIdentity::EPrivilegeResults::UserNotLoggedIn) ResultNames.Emplace(TEXT("UserNotLoggedIn"));
+		if ((uint32)PrivilegeResults & (uint32)IOnlineIdentity::EPrivilegeResults::ChatRestriction) ResultNames.Emplace(TEXT("ChatRestriction"));
+		if ((uint32)PrivilegeResults & (uint32)IOnlineIdentity::EPrivilegeResults::UGCRestriction) ResultNames.Emplace(TEXT("UGCRestriction"));
+		if ((uint32)PrivilegeResults & (uint32)IOnlineIdentity::EPrivilegeResults::GenericFailure) ResultNames.Emplace(TEXT("GenericFailure"));
+		if ((uint32)PrivilegeResults & (uint32)IOnlineIdentity::EPrivilegeResults::OnlinePlayRestricted) ResultNames.Emplace(TEXT("OnlinePlayRestricted"));
+		if ((uint32)PrivilegeResults & (uint32)IOnlineIdentity::EPrivilegeResults::NetworkConnectionUnavailable) ResultNames.Emplace(TEXT("NetworkConnectionUnavailable"));
+		check((uint32)PrivilegeResults < ((uint32)IOnlineIdentity::EPrivilegeResults::NetworkConnectionUnavailable << 1));
+
+		return FString::Join(ResultNames, TEXT(" | "));
+	}
+}
+
+FString ToDebugString(EUserPrivileges::Type UserPrivilege)
+{
+	switch (UserPrivilege)
+	{
+	case EUserPrivileges::CanPlay: return TEXT("CanPlay");
+	case EUserPrivileges::CanPlayOnline: return TEXT("CanPlayOnline");
+	case EUserPrivileges::CanCommunicateOnline: return TEXT("CanCommunicateOnline");
+	case EUserPrivileges::CanUseUserGeneratedContent: return TEXT("CanUseUserGeneratedContent");
+	case EUserPrivileges::CanUserCrossPlay: return TEXT("CanUserCrossPlay");
+	default: return TEXT("Unknown");
+	}
+}
\ No newline at end of file
diff --git a/Engine/Plugins/Online/OnlineSubsystem/Source/Public/Interfaces/OnlineIdentityInterface.h b/Engine/Plugins/Online/OnlineSubsystem/Source/Public/Interfaces/OnlineIdentityInterface.h
--- a/Engine/Plugins/Online/OnlineSubsystem/Source/Public/Interfaces/OnlineIdentityInterface.h
+++ b/Engine/Plugins/Online/OnlineSubsystem/Source/Public/Interfaces/OnlineIdentityInterface.h
@@ -446,3 +446,7 @@ public:
 };
 
 typedef TSharedPtr<IOnlineIdentity, ESPMode::ThreadSafe> IOnlineIdentityPtr;
+
+ONLINESUBSYSTEM_API FString ToDebugString(IOnlineIdentity::EPrivilegeResults PrivilegeResult);
+
+ONLINESUBSYSTEM_API FString ToDebugString(EUserPrivileges::Type UserPrivilege);
\ No newline at end of file
diff --git a/Engine/Plugins/Online/OnlineSubsystem/Source/Public/OnlineAsyncTaskManager.h b/Engine/Plugins/Online/OnlineSubsystem/Source/Public/OnlineAsyncTaskManager.h
--- a/Engine/Plugins/Online/OnlineSubsystem/Source/Public/OnlineAsyncTaskManager.h
+++ b/Engine/Plugins/Online/OnlineSubsystem/Source/Public/OnlineAsyncTaskManager.h
@@ -186,9 +186,10 @@ public:
 	*
 	* @param InCallable any object that can be called with no parameters, usually a lambda
 	*/
-	explicit FOnlineAsyncTaskThreadedGenericCallable(const CallableType& InCallable)
+	explicit FOnlineAsyncTaskThreadedGenericCallable(const FString& InCallableName, const CallableType& InCallable)
 		: bHasTicked(false)
 		, CallableObject(InCallable)
+		, CallableName(InCallableName)
 	{
 	}
 
@@ -198,7 +199,7 @@ public:
 		bHasTicked = true;
 	}
 
-	virtual FString ToString() const override { return FString("FOnlineAsyncTaskThreadedGenericCallable"); }
+	virtual FString ToString() const override { return FString::Printf(TEXT("FOnlineAsyncTaskThreadedGenericCallable (%s)"), *CallableName); }
 
 	virtual bool IsDone() const override { return bHasTicked; }
 	virtual bool WasSuccessful() const override { return true; }
@@ -208,6 +209,8 @@ private:
 	bool bHasTicked;
 	/** Stored copy of the object to invoke on the game thread. */
 	CallableType CallableObject;
+	/** Name of the task to help logging */
+	FString CallableName;
 };
 
 template<class T>
@@ -431,9 +434,9 @@ public:
 	* @param InCallable the callable object to execute on the game thread.
 	*/
 	template<class CallableType>
-	void AddGenericToInQueueOnlineThread(const CallableType& InCallable)
+	void AddGenericToInQueueOnlineThread(const FString& CallableName, const CallableType& InCallable)
 	{
-		AddToInQueue(new FOnlineAsyncTaskThreadedGenericCallable<CallableType>(InCallable));
+		AddToInQueue(new FOnlineAsyncTaskThreadedGenericCallable<CallableType>(CallableName, InCallable));
 	}
 
 	/**
