- Add MismatchedApp JoinPartyCompleteCode to disambiguate from BadBuild - Split BadBuild subcodes into BuildIdTooLow/High, CompatibilityVersionTooLow/High

diff --git a/Engine/Plugins/Online/OnlineSubsystem/Source/Private/OnlinePartyInterface.cpp b/Engine/Plugins/Online/OnlineSubsystem/Source/Private/OnlinePartyInterface.cpp
--- a/Engine/Plugins/Online/OnlineSubsystem/Source/Private/OnlinePartyInterface.cpp
+++ b/Engine/Plugins/Online/OnlineSubsystem/Source/Private/OnlinePartyInterface.cpp
@@ -541,6 +541,10 @@ const TCHAR* ToString(const EJoinPartyCompletionResult Value)
 	{
 		return TEXT("GameSpecificReason");
 	}
+	case EJoinPartyCompletionResult::MismatchedApp:
+	{
+		return TEXT("MismatchedApp");
+	}
 	case EJoinPartyCompletionResult::Succeeded:
 	{
 		return TEXT("Succeeded");
diff --git a/Engine/Plugins/Online/OnlineSubsystem/Source/Public/Interfaces/OnlinePartyInterface.h b/Engine/Plugins/Online/OnlineSubsystem/Source/Public/Interfaces/OnlinePartyInterface.h
--- a/Engine/Plugins/Online/OnlineSubsystem/Source/Public/Interfaces/OnlinePartyInterface.h
+++ b/Engine/Plugins/Online/OnlineSubsystem/Source/Public/Interfaces/OnlinePartyInterface.h
@@ -1932,6 +1932,9 @@ enum class EJoinPartyCompletionResult : int8
 
 	/** Game specific reason, indicated by the NotApprovedReason parameter.  Message might or might not have been sent to party leader. */
 	GameSpecificReason,
+	
+	/** Your app id does not match the app id of the party */
+	MismatchedApp,
 
 	/** DEPRECATED */
 	UnknownInternalFailure = 0,
