Modify FOnlineEventParms to use FString instead of FName (OnlineSubsystem)

diff --git a/Engine/Plugins/Online/OnlineSubsystem/Source/Public/Interfaces/OnlineEventsInterface.h b/Engine/Plugins/Online/OnlineSubsystem/Source/Public/Interfaces/OnlineEventsInterface.h
--- a/Engine/Plugins/Online/OnlineSubsystem/Source/Public/Interfaces/OnlineEventsInterface.h
+++ b/Engine/Plugins/Online/OnlineSubsystem/Source/Public/Interfaces/OnlineEventsInterface.h
@@ -17,8 +17,60 @@ ONLINESUBSYSTEM_API DECLARE_LOG_CATEGORY_EXTERN(LogOnlineEvents, Log, All);
 }
 
 class FUniqueNetId;
+// Temporary class to assist in deprecation of changing the key from FName to FString. This will eventually be a typedef to FOnlineKeyValuePairs< FString, FVariantData >
+class FOnlineEventParms : public FOnlineKeyValuePairs< FString, FVariantData >
+{
+	typedef FOnlineKeyValuePairs<FString, FVariantData> Super;
+
+public:
+	FORCEINLINE FOnlineEventParms() {}
+	FORCEINLINE FOnlineEventParms(FOnlineEventParms&& Other) : Super(MoveTemp(Other)) {}
+	FORCEINLINE FOnlineEventParms(const FOnlineEventParms& Other) : Super(Other) {}
+	FORCEINLINE FOnlineEventParms& operator=(FOnlineEventParms&& Other) { Super::operator=(MoveTemp(Other)); return *this; }
+	FORCEINLINE FOnlineEventParms& operator=(const FOnlineEventParms& Other) { Super::operator=(Other); return *this; }
+
+public:
+	UE_DEPRECATED(4.26, "FOnlineEventParms now uses FString for the key type instead of FName")
+		FOnlineEventParms(const FOnlineKeyValuePairs< FName, FVariantData >& DeprecatedValues)
+	{
+		Reserve(DeprecatedValues.Num());
+		for (const TPair< FName, FVariantData >& DeprecatedValue : DeprecatedValues)
+		{
+			Super::Emplace(DeprecatedValue.Key.ToString(), DeprecatedValue.Value);
+		}
+	}
 
-typedef FOnlineKeyValuePairs< FName, FVariantData > FOnlineEventParms;
+	UE_DEPRECATED(4.26, "FOnlineEventParms now uses FString for the key type instead of FName")
+		FOnlineEventParms(FOnlineKeyValuePairs< FName, FVariantData >&& DeprecatedValues)
+	{
+		Reserve(DeprecatedValues.Num());
+		for (TPair< FName, FVariantData >& DeprecatedValue : DeprecatedValues)
+		{
+			Super::Emplace(DeprecatedValue.Key.ToString(), MoveTemp(DeprecatedValue.Value));
+		}
+	}
+ 	using Super::Add;
+ 	UE_DEPRECATED(4.26, "FOnlineEventParms now uses FString for the key type instead of FName")
+ 		FVariantData& Add(const FName& InKey, const FVariantData& InValue)
+ 	{
+ 		return Super::Add(InKey.ToString(), InValue);
+ 	}
+
+	using Super::Emplace;
+	UE_DEPRECATED(4.26, "FOnlineEventParms now uses FString for the key type instead of FName")
+		FVariantData& Emplace(const FName& InKey, const FVariantData& InValue)
+	{
+		return Super::Emplace(InKey.ToString(), InValue);
+	}
+
+	using Super::Find;
+	UE_DEPRECATED(4.26, "FOnlineEventParms now uses FString for the key type instead of FName")
+		FVariantData* Find(const FName& InKey)
+	{
+		return Super::Find(InKey.ToString());
+	}
+};
+// typedef FOnlineKeyValuePairs< FString, FVariantData > FOnlineEventParms;
 
 /**
  *	IOnlineEvents - Interface class for events
@@ -45,4 +97,4 @@ public:
 	 * @param PlayerSessionId A GUID unique to this player session
 	 */
 	virtual void SetPlayerSessionId(const FUniqueNetId& PlayerId, const FGuid& PlayerSessionId) = 0;
-};
+};
\ No newline at end of file
