- Deprecate FUniqueNetId public constructors

diff --git a/Engine/Plugins/Online/OnlineSubsystem/Source/Public/OnlineSubsystemTypes.h b/Engine/Plugins/Online/OnlineSubsystem/Source/Public/OnlineSubsystemTypes.h
--- a/Engine/Plugins/Online/OnlineSubsystem/Source/Public/OnlineSubsystemTypes.h
+++ b/Engine/Plugins/Online/OnlineSubsystem/Source/Public/OnlineSubsystemTypes.h
@@ -832,12 +832,11 @@ public:
 	template<typename... TArgs>
 	static FUniqueNetIdStringRef Create(TArgs&&... Args)
 	{
-		return MakeShared<FUniqueNetIdString>(Forward<TArgs>(Args)...);
+		PRAGMA_DISABLE_DEPRECATION_WARNINGS
+		return MakeShareable(new FUniqueNetIdString(Forward<TArgs>(Args)...));
+		PRAGMA_ENABLE_DEPRECATION_WARNINGS
 	}
 
-	/** Allow MakeShared to see private constructors */
-	friend class SharedPointerInternals::TIntrusiveReferenceController<FUniqueNetIdString>;
-
 	static FUniqueNetIdStringRef& EmptyId()
 	{
 		static FUniqueNetIdStringRef EmptyId(Create());
@@ -892,6 +891,7 @@ public:
 	}
 
 public:
+	UE_DEPRECATED(5.0, "Public constructors of FUniqueNetId types are deprecated. Please use the ::Create(Args) method instead to create a FUniqueNetIdRef")
 	FUniqueNetIdString() = default;
 
 	/**
@@ -899,7 +899,8 @@ public:
 	 *
 	 * @param InUniqueNetId the id to set ours to
 	 */
-	explicit FUniqueNetIdString(const FString & InUniqueNetId)
+	UE_DEPRECATED(5.0, "Public constructors of FUniqueNetId types are deprecated. Please use the ::Create(Args) method instead to create a FUniqueNetIdRef")
+	explicit FUniqueNetIdString(const FString& InUniqueNetId)
 		: UniqueNetIdStr(InUniqueNetId)
 		, Type(NAME_Unset)
 	{
@@ -910,7 +911,8 @@ public:
 	 *
 	 * @param InUniqueNetId the id to set ours to
 	 */
-	explicit FUniqueNetIdString(FString && InUniqueNetId)
+	UE_DEPRECATED(5.0, "Public constructors of FUniqueNetId types are deprecated. Please use the ::Create(Args) method instead to create a FUniqueNetIdRef")
+	explicit FUniqueNetIdString(FString&& InUniqueNetId)
 		: UniqueNetIdStr(MoveTemp(InUniqueNetId))
 		, Type(NAME_Unset)
 	{
@@ -921,7 +923,8 @@ public:
 	 *
 	 * @param Src the id to copy
 	 */
-	explicit FUniqueNetIdString(const FUniqueNetId & Src)
+	UE_DEPRECATED(5.0, "Public constructors of FUniqueNetId types are deprecated. Please use the ::Create(Args) method instead to create a FUniqueNetIdRef")
+	explicit FUniqueNetIdString(const FUniqueNetId& Src)
 		: UniqueNetIdStr(Src.ToString())
 		, Type(Src.GetType())
 	{
@@ -930,7 +933,8 @@ public:
 	/**
 	* don.eubanks - Including a constructor that allows for type passing to make transitioning easier, if we determine we want to abstract-ify this class, this constructor will be removed
 	*/
-	FUniqueNetIdString(const FString & InUniqueNetId, const FName InType)
+	UE_DEPRECATED(5.0, "Public constructors of FUniqueNetId types are deprecated. Please use the ::Create(Args) method instead to create a FUniqueNetIdRef")
+	FUniqueNetIdString(const FString& InUniqueNetId, const FName InType)
 		: UniqueNetIdStr(InUniqueNetId)
 		, Type(InType)
 	{
@@ -947,9 +951,10 @@ public: \
 	template<typename... TArgs> \
 	static SUBCLASSNAME##Ref Create(TArgs&&... Args) \
 	{ \
-		return MakeShared<SUBCLASSNAME>(Forward<TArgs>(Args)...); \
+		PRAGMA_DISABLE_DEPRECATION_WARNINGS \
+		return MakeShareable(new SUBCLASSNAME(Forward<TArgs>(Args)...)); \
+		PRAGMA_ENABLE_DEPRECATION_WARNINGS \
 	} \
-	friend class SharedPointerInternals::TIntrusiveReferenceController<SUBCLASSNAME>; \
 	static SUBCLASSNAME##Ref Cast(const FUniqueNetIdRef& InNetId) \
 	{ \
 		check(InNetId->GetType() == TYPE); \
@@ -983,19 +988,23 @@ public: \
 		return EmptyId; \
 	} \
 public: \
+	UE_DEPRECATED(5.0, "Public constructors of FUniqueNetId types are deprecated. Please use the ::Create(Args) method instead to create a FUniqueNetIdRef") \
 	SUBCLASSNAME() \
 		: FUniqueNetIdString() \
 	{ \
 		Type = TYPE; \
 	} \
+	UE_DEPRECATED(5.0, "Public constructors of FUniqueNetId types are deprecated. Please use the ::Create(Args) method instead to create a FUniqueNetIdRef") \
 	explicit SUBCLASSNAME(const FString& InUniqueNetId) \
 		: FUniqueNetIdString(InUniqueNetId, TYPE) \
 	{ \
 	} \
+	UE_DEPRECATED(5.0, "Public constructors of FUniqueNetId types are deprecated. Please use the ::Create(Args) method instead to create a FUniqueNetIdRef") \
 	explicit SUBCLASSNAME(FString&& InUniqueNetId) \
 		: FUniqueNetIdString(MoveTemp(InUniqueNetId), TYPE) \
 	{ \
 	} \
+	UE_DEPRECATED(5.0, "Public constructors of FUniqueNetId types are deprecated. Please use the ::Create(Args) method instead to create a FUniqueNetIdRef") \
 	explicit SUBCLASSNAME(const FUniqueNetId& Src) \
 		: FUniqueNetIdString(Src) \
 	{ \
diff --git a/Engine/Plugins/Online/OnlineSubsystemEOS/Source/OnlineSubsystemEOS/Private/OnlineSubsystemEOSTypes.h b/Engine/Plugins/Online/OnlineSubsystemEOS/Source/OnlineSubsystemEOS/Private/OnlineSubsystemEOSTypes.h
--- a/Engine/Plugins/Online/OnlineSubsystemEOS/Source/OnlineSubsystemEOS/Private/OnlineSubsystemEOSTypes.h
+++ b/Engine/Plugins/Online/OnlineSubsystemEOS/Source/OnlineSubsystemEOS/Private/OnlineSubsystemEOSTypes.h
@@ -46,12 +46,9 @@ public:
 	template<typename... TArgs>
 	static FUniqueNetIdEOSRef Create(TArgs&&... Args)
 	{
-		return MakeShared<FUniqueNetIdEOS>(Forward<TArgs>(Args)...);
+		return MakeShareable(new FUniqueNetIdEOS(Forward<TArgs>(Args)...));
 	}
 
-	/** Allow MakeShared to see private constructors */
-	friend class SharedPointerInternals::TIntrusiveReferenceController<FUniqueNetIdEOS>;
-
 	static const FUniqueNetIdEOS& Cast(const FUniqueNetId& NetId)
 	{
 		check(GetTypeStatic() == NetId.GetType());
diff --git a/Engine/Plugins/Online/OnlineSubsystemEOS/Source/OnlineSubsystemEOSPlus/Private/OnlineUserEOSPlus.h b/Engine/Plugins/Online/OnlineSubsystemEOS/Source/OnlineSubsystemEOSPlus/Private/OnlineUserEOSPlus.h
--- a/Engine/Plugins/Online/OnlineSubsystemEOS/Source/OnlineSubsystemEOSPlus/Private/OnlineUserEOSPlus.h
+++ b/Engine/Plugins/Online/OnlineSubsystemEOS/Source/OnlineSubsystemEOSPlus/Private/OnlineUserEOSPlus.h
@@ -31,12 +31,9 @@ public:
 	template<typename... TArgs>
 	static FUniqueNetIdEOSPlusRef Create(TArgs&&... Args)
 	{
-		return MakeShared<FUniqueNetIdEOSPlus>(Forward<TArgs>(Args)...);
+		return MakeShareable(new FUniqueNetIdEOSPlus(Forward<TArgs>(Args)...));
 	}
 
-	/** Allow MakeShared to see private constructors */
-	friend class SharedPointerInternals::TIntrusiveReferenceController<FUniqueNetIdEOSPlus>;
-
 // FUniqueNetId interface
 	virtual const uint8* GetBytes() const override;
 	virtual int32 GetSize() const override;
@@ -73,12 +70,9 @@ public:
 	template<typename... TArgs>
 	static FUniqueNetIdBinaryRef Create(TArgs&&... Args)
 	{
-		return MakeShared<FUniqueNetIdBinary>(Forward<TArgs>(Args)...);
+		return MakeShareable(new FUniqueNetIdBinary(Forward<TArgs>(Args)...));
 	}
 
-	/** Allow MakeShared to see private constructors */
-	friend class SharedPointerInternals::TIntrusiveReferenceController<FUniqueNetIdBinary>;
-
 	virtual FName GetType() const override
 	{
 		return Type;
diff --git a/Engine/Plugins/Online/OnlineSubsystemFacebook/Source/Private/OnlineSubsystemFacebookTypes.h b/Engine/Plugins/Online/OnlineSubsystemFacebook/Source/Private/OnlineSubsystemFacebookTypes.h
--- a/Engine/Plugins/Online/OnlineSubsystemFacebook/Source/Private/OnlineSubsystemFacebookTypes.h
+++ b/Engine/Plugins/Online/OnlineSubsystemFacebook/Source/Private/OnlineSubsystemFacebookTypes.h
@@ -27,12 +27,9 @@ public:
 	template<typename... TArgs>
 	static FUniqueNetIdFacebookRef Create(TArgs&&... Args)
 	{
-		return MakeShared<FUniqueNetIdFacebook>(Forward<TArgs>(Args)...);
+		return MakeShareable(new FUniqueNetIdFacebook(Forward<TArgs>(Args)...));
 	}
 
-	/** Allow MakeShared to see private constructors */
-	friend class SharedPointerInternals::TIntrusiveReferenceController<FUniqueNetIdFacebook>;
-
 	//~ Begin FUniqueNetId Interface
 	virtual FName GetType() const override
 	{
diff --git a/Engine/Plugins/Online/OnlineSubsystemOculus/Source/Private/OnlineSubsystemOculusTypes.h b/Engine/Plugins/Online/OnlineSubsystemOculus/Source/Private/OnlineSubsystemOculusTypes.h
--- a/Engine/Plugins/Online/OnlineSubsystemOculus/Source/Private/OnlineSubsystemOculusTypes.h
+++ b/Engine/Plugins/Online/OnlineSubsystemOculus/Source/Private/OnlineSubsystemOculusTypes.h
@@ -34,12 +34,9 @@ public:
 	template<typename... TArgs>
 	static FUniqueNetIdOculusRef Create(TArgs&&... Args)
 	{
-		return MakeShared<FUniqueNetIdOculus>(Forward<TArgs>(Args)...);
+		return MakeShareable(new FUniqueNetIdOculus(Forward<TArgs>(Args)...));
 	}
 
-	/** Allow MakeShared to see private constructors */
-	friend class SharedPointerInternals::TIntrusiveReferenceController<FUniqueNetIdOculus>;
-
 	static const FUniqueNetIdOculus& Cast(const FUniqueNetId& NetId)
 	{
 		check(NetId.GetType() == OCULUS_SUBSYSTEM);
diff --git a/Engine/Plugins/Online/OnlineSubsystemSteam/Source/Private/OnlineSubsystemSteamTypes.h b/Engine/Plugins/Online/OnlineSubsystemSteam/Source/Private/OnlineSubsystemSteamTypes.h
--- a/Engine/Plugins/Online/OnlineSubsystemSteam/Source/Private/OnlineSubsystemSteamTypes.h
+++ b/Engine/Plugins/Online/OnlineSubsystemSteam/Source/Private/OnlineSubsystemSteamTypes.h
@@ -73,7 +73,7 @@ public:
 	template<typename... TArgs>
 	static FUniqueNetIdSteamRef Create(TArgs&&... Args)
 	{
-		return MakeShared<FUniqueNetIdSteam>(Forward<TArgs>(Args)...);
+		return MakeShareable(new FUniqueNetIdSteam(Forward<TArgs>(Args)...));
 	}
 
 	static const FUniqueNetIdSteam& Cast(const FUniqueNetId& NetId)
@@ -86,9 +86,6 @@ public:
 		return StaticCastSharedRef<const FUniqueNetIdSteam>(FUniqueNetId::AsShared());
 	}
 
-	/** Allow MakeShared to see private constructors */
-	friend class SharedPointerInternals::TIntrusiveReferenceController<FUniqueNetIdSteam>;
-
 	virtual FName GetType() const override
 	{
 		return STEAM_SUBSYSTEM;
diff --git a/Engine/Plugins/Online/OnlineSubsystemTencent/Source/Public/OnlineSubsystemTencentTypes.h b/Engine/Plugins/Online/OnlineSubsystemTencent/Source/Public/OnlineSubsystemTencentTypes.h
--- a/Engine/Plugins/Online/OnlineSubsystemTencent/Source/Public/OnlineSubsystemTencentTypes.h
+++ b/Engine/Plugins/Online/OnlineSubsystemTencent/Source/Public/OnlineSubsystemTencentTypes.h
@@ -47,12 +47,9 @@ public:
 	template<typename... TArgs>
 	static FUniqueNetIdRailRef Create(TArgs&&... Args)
 	{
-		return MakeShared<FUniqueNetIdRail>(Forward<TArgs>(Args)...);
+		return MakeShareable(new FUniqueNetIdRail(Forward<TArgs>(Args)...));
 	}
 
-	/** Allow MakeShared to see private constructors */
-	friend class SharedPointerInternals::TIntrusiveReferenceController<FUniqueNetIdRail>;
-
 	virtual FName GetType() const override
 	{
 		return TENCENT_SUBSYSTEM;
diff --git a/Engine/Source/Runtime/CoreOnline/CoreOnline.Build.cs b/Engine/Source/Runtime/CoreOnline/CoreOnline.Build.cs
--- a/Engine/Source/Runtime/CoreOnline/CoreOnline.Build.cs
+++ b/Engine/Source/Runtime/CoreOnline/CoreOnline.Build.cs
@@ -5,9 +5,6 @@ using EpicGames.Core;
 
 public class CoreOnline : ModuleRules
 {
-	[ConfigFile(ConfigHierarchyType.Engine, "CoreOnline")]
-	bool bEnableThreadSafeUniqueNetIds = false;
-
 	public CoreOnline(ReadOnlyTargetRules Target) : base(Target)
 	{
 		PublicDependencyModuleNames.AddRange(
@@ -16,8 +13,5 @@ public class CoreOnline : ModuleRules
 				"CoreUObject"
 			}
 		);
-
-		ConfigCache.ReadSettings(DirectoryReference.FromFile(Target.ProjectFile), Target.Platform, this);
 	}
-
 }
