Change default OSS polling interval to 20ms, this matches the faster value required by all internal games and several platforms

diff --git a/Engine/Config/BaseEngine.ini b/Engine/Config/BaseEngine.ini
--- a/Engine/Config/BaseEngine.ini
+++ b/Engine/Config/BaseEngine.ini
@@ -1728,7 +1728,7 @@ VoiceNotificationDelta=0.33
 ;VoiceNotificationDelta=0.2
 MaxLocalTalkers=1
 MaxRemoteTalkers=16
-PollingIntervalInMs=50
+PollingIntervalInMs=20
 bUseBuildIdOverride=false
 BuildIdOverride=0
 !AdditionalModulesToLoad=Clear
diff --git a/Engine/Plugins/Online/OnlineSubsystem/Source/Private/OnlineAsyncTaskManager.cpp b/Engine/Plugins/Online/OnlineSubsystem/Source/Private/OnlineAsyncTaskManager.cpp
--- a/Engine/Plugins/Online/OnlineSubsystem/Source/Private/OnlineAsyncTaskManager.cpp
+++ b/Engine/Plugins/Online/OnlineSubsystem/Source/Private/OnlineAsyncTaskManager.cpp
@@ -25,7 +25,7 @@ namespace OSSConsoleVariables
 #endif
 
 /** The default value for the polling interval when not set by config */
-#define POLLING_INTERVAL_MS 50
+#define POLLING_INTERVAL_MS 20
 
 FOnlineAsyncTaskManager::FOnlineAsyncTaskManager() :
 	ActiveTask(nullptr),
