Added API method for clearing cached auth token (where applicable) in online identity interface.

diff --git a/Engine/Plugins/Online/OnlineSubsystem/Source/Public/Interfaces/OnlineIdentityInterface.h b/Engine/Plugins/Online/OnlineSubsystem/Source/Public/Interfaces/OnlineIdentityInterface.h
--- a/Engine/Plugins/Online/OnlineSubsystem/Source/Public/Interfaces/OnlineIdentityInterface.h
+++ b/Engine/Plugins/Online/OnlineSubsystem/Source/Public/Interfaces/OnlineIdentityInterface.h
@@ -434,6 +434,18 @@ public:
 	 */
 	virtual void RevokeAuthToken(const FUniqueNetId& LocalUserId, const FOnRevokeAuthTokenCompleteDelegate& Delegate) = 0;
 
+	/**
+	 * Clear an auth token associated with the user, if one was cached.
+	 *
+	 * If an auth token is obtained for some other purpose than user account login, it would
+	 * be ignored an replaced with another when an actual login attempt is made, wasting the
+	 * previously obtained one.  Caching avoids that potential waste, but must be manually
+	 * cleared client-side when the login completes due to the remote side has consuming it.
+	 *
+	 * @param UserId the unique net of the associated user
+	 */
+	virtual void ClearCachedAuthToken(const FUniqueNetId& UserId) {};
+
 	/**
 	 * Delegate executed when the user's auth ticket has arrived
 	 *
