Add ToDebugString for FControllerPairingChangedUserInfo

diff --git a/Engine/Plugins/Online/OnlineSubsystem/Source/Private/OnlineIdentityInterface.cpp b/Engine/Plugins/Online/OnlineSubsystem/Source/Private/OnlineIdentityInterface.cpp
--- a/Engine/Plugins/Online/OnlineSubsystem/Source/Private/OnlineIdentityInterface.cpp
+++ b/Engine/Plugins/Online/OnlineSubsystem/Source/Private/OnlineIdentityInterface.cpp
@@ -42,6 +42,11 @@ FString ToDebugString(EUserPrivileges::Type UserPrivilege)
 	}
 }
 
+FString ToDebugString(const FControllerPairingChangedUserInfo& ControllerPairingChangedUserInfo)
+{
+	return FString::Printf(TEXT("[%s, %d devices]"), *ControllerPairingChangedUserInfo.User.ToDebugString(), ControllerPairingChangedUserInfo.ControllersRemaining);
+}
+
 FPlatformUserId IOnlineIdentity::GetPlatformUserIdFromLocalUserNum(int32 LocalUserNum) const
 {
 	if (LocalUserNum >= 0 && LocalUserNum < MAX_LOCAL_PLAYERS)
diff --git a/Engine/Plugins/Online/OnlineSubsystem/Source/Public/Interfaces/OnlineIdentityInterface.h b/Engine/Plugins/Online/OnlineSubsystem/Source/Public/Interfaces/OnlineIdentityInterface.h
--- a/Engine/Plugins/Online/OnlineSubsystem/Source/Public/Interfaces/OnlineIdentityInterface.h
+++ b/Engine/Plugins/Online/OnlineSubsystem/Source/Public/Interfaces/OnlineIdentityInterface.h
@@ -523,4 +523,6 @@ typedef TSharedPtr<IOnlineIdentity, ESPMode::ThreadSafe> IOnlineIdentityPtr;
 
 ONLINESUBSYSTEM_API FString ToDebugString(IOnlineIdentity::EPrivilegeResults PrivilegeResult);
 
-ONLINESUBSYSTEM_API FString ToDebugString(EUserPrivileges::Type UserPrivilege);
\ No newline at end of file
+ONLINESUBSYSTEM_API FString ToDebugString(EUserPrivileges::Type UserPrivilege);
+
+ONLINESUBSYSTEM_API FString ToDebugString(const FControllerPairingChangedUserInfo& ControllerPairingChangedUserInfo);
