PR #8587
Added session owner check to IsPlayerInSessionImpl in OnlineSubsystem.
(Contributed by [REDACTED_PERSON])

diff --git a/Engine/Plugins/Online/OnlineSubsystem/Source/Private/OnlineSubsystem.cpp b/Engine/Plugins/Online/OnlineSubsystem/Source/Private/OnlineSubsystem.cpp
--- a/Engine/Plugins/Online/OnlineSubsystem/Source/Private/OnlineSubsystem.cpp
+++ b/Engine/Plugins/Online/OnlineSubsystem/Source/Private/OnlineSubsystem.cpp
@@ -224,7 +224,7 @@ bool IsPlayerInSessionImpl(IOnlineSession* SessionInt, FName SessionName, const
 	FNamedOnlineSession* Session = SessionInt->GetNamedSession(SessionName);
 	if (Session != NULL)
 	{
-		const bool bIsSessionOwner = *Session->OwningUserId == UniqueId;
+		const bool bIsSessionOwner = Session->OwningUserId.IsValid()  && *Session->OwningUserId == UniqueId;
 
 		FUniqueNetIdMatcher PlayerMatch(UniqueId);
 		if (bIsSessionOwner || 
