Convert UpdateFriendSettings HttpRequest to use PATCH instead of PUT
#rnx
#tests Made sure friend invite privacy setting persisted upon closing and reopening client

Rename UpdateFriendSettings to UpdateFriendInvitePolicySettings and use explicit parameters to avoid converting to and from an unnecessary dictionary format.
Add affectExistingInvites query param that wasn't previously passed to the backend

diff --git a/Engine/Plugins/Online/OnlineSubsystem/Source/Public/Interfaces/OnlineFriendsInterface.h b/Engine/Plugins/Online/OnlineSubsystem/Source/Public/Interfaces/OnlineFriendsInterface.h
--- a/Engine/Plugins/Online/OnlineSubsystem/Source/Public/Interfaces/OnlineFriendsInterface.h
+++ b/Engine/Plugins/Online/OnlineSubsystem/Source/Public/Interfaces/OnlineFriendsInterface.h
@@ -717,11 +717,12 @@ public:
 	 * Update the current friend settings
 	 *
 	 * @param LocalUserId user to retrieve friend settings for
-	 * @param NewSettings Settings to be saved
+	 * @param FriendInvitePolicy invite policy to change designated user to
+	 * @param bAffectsExistingInvites boolean to determine if existing friend invites are affected by this change
 	 * @param Delegate Delegate to call when operation has been completed
 	 *
 	 */
-	virtual void UpdateFriendSettings(const FUniqueNetId& LocalUserId, const FFriendSettings& NewSettings, FOnSettingsOperationComplete Delegate) { check(false) }
+	virtual void UpdateFriendInvitePolicySettings(const FUniqueNetId& LocalUserId, const EFriendInvitePolicy FriendInvitePolicy, const bool bAffectsExistingInvites, FOnSettingsOperationComplete Delegate) { check(false) }
 
 	/**
 	 * Queries the settings we have stored for a third party source
