Move MAX_LOCAL_PLAYERS from OnlineSubsystem plugin to CoreOnline

diff --git a/Engine/Plugins/Online/OnlineSubsystem/Source/OnlineSubsystem.Build.cs b/Engine/Plugins/Online/OnlineSubsystem/Source/OnlineSubsystem.Build.cs
--- a/Engine/Plugins/Online/OnlineSubsystem/Source/OnlineSubsystem.Build.cs
+++ b/Engine/Plugins/Online/OnlineSubsystem/Source/OnlineSubsystem.Build.cs
@@ -20,7 +20,6 @@ public class OnlineSubsystem : ModuleRules
 
         PublicDefinitions.Add("ONLINESUBSYSTEM_PACKAGE=1");
 		PublicDefinitions.Add("DEBUG_LAN_BEACON=0");
-		PublicDefinitions.Add("PLATFORM_MAX_LOCAL_PLAYERS=" + GetPlatformMaxLocalPlayers(Target));
 
 		// OnlineSubsystem cannot depend on Engine!
 		PrivateDependencyModuleNames.AddRange(
@@ -32,10 +31,4 @@ public class OnlineSubsystem : ModuleRules
 			}
 		);
 	}
-
-	protected virtual int GetPlatformMaxLocalPlayers(ReadOnlyTargetRules Target)
-	{
-		// 0 indicates no platform override
-		return 0;
-	}
 }
diff --git a/Engine/Plugins/Online/OnlineSubsystem/Source/Public/OnlineSubsystemTypes.h b/Engine/Plugins/Online/OnlineSubsystem/Source/Public/OnlineSubsystemTypes.h
--- a/Engine/Plugins/Online/OnlineSubsystem/Source/Public/OnlineSubsystemTypes.h
+++ b/Engine/Plugins/Online/OnlineSubsystem/Source/Public/OnlineSubsystemTypes.h
@@ -23,17 +23,6 @@ extern ONLINESUBSYSTEM_API bool IsUniqueIdLocal(const FUniqueNetId& UniqueId);
 #define OSS_REDACT(x) (x)
 #endif
 
-/** Maximum players supported on a given platform */
-#if !defined(MAX_LOCAL_PLAYERS)
-	#if PLATFORM_MAX_LOCAL_PLAYERS
-		#define MAX_LOCAL_PLAYERS PLATFORM_MAX_LOCAL_PLAYERS
-	#elif PLATFORM_WINDOWS
-		#define MAX_LOCAL_PLAYERS 4
-	#else
-		#define MAX_LOCAL_PLAYERS 1
-	#endif
-#endif //MAX_LOCAL_PLAYERS
-
 #define DEDICATED_SERVER_USER_INDEX 0
 
 #ifndef ONLINE_SUCCESS
diff --git a/Engine/Source/Runtime/CoreOnline/CoreOnline.Build.cs b/Engine/Source/Runtime/CoreOnline/CoreOnline.Build.cs
--- a/Engine/Source/Runtime/CoreOnline/CoreOnline.Build.cs
+++ b/Engine/Source/Runtime/CoreOnline/CoreOnline.Build.cs
@@ -15,5 +15,12 @@ public class CoreOnline : ModuleRules
 		);
 
 		PrivateDefinitions.Add("COREONLINE_PACKAGE=1");
+		PublicDefinitions.Add("PLATFORM_MAX_LOCAL_PLAYERS=" + GetPlatformMaxLocalPlayers(Target));
+	}
+
+	protected virtual int GetPlatformMaxLocalPlayers(ReadOnlyTargetRules Target)
+	{
+		// 0 indicates no platform override
+		return 0;
 	}
 }
diff --git a/Engine/Source/Runtime/CoreOnline/Public/Online/CoreOnline.h b/Engine/Source/Runtime/CoreOnline/Public/Online/CoreOnline.h
--- a/Engine/Source/Runtime/CoreOnline/Public/Online/CoreOnline.h
+++ b/Engine/Source/Runtime/CoreOnline/Public/Online/CoreOnline.h
@@ -12,6 +12,17 @@
 
 class FLazySingleton;
 
+/** Maximum players supported on a given platform */
+#if !defined(MAX_LOCAL_PLAYERS)
+	#if PLATFORM_MAX_LOCAL_PLAYERS
+		#define MAX_LOCAL_PLAYERS PLATFORM_MAX_LOCAL_PLAYERS
+	#elif PLATFORM_DESKTOP
+		#define MAX_LOCAL_PLAYERS 4
+	#else
+		#define MAX_LOCAL_PLAYERS 1
+	#endif
+#endif //MAX_LOCAL_PLAYERS
+
 UE_DEPRECATED(5.0, "Use NAME_GameSession.")
 inline constexpr EName GameSessionName = NAME_GameSession;
 UE_DEPRECATED(5.0, "Use NAME_PartySession.")
