Removing reference to UE3

#rnx

diff --git a/Engine/Plugins/Online/OnlineSubsystem/Source/Public/LANBeacon.h b/Engine/Plugins/Online/OnlineSubsystem/Source/Public/LANBeacon.h
--- a/Engine/Plugins/Online/OnlineSubsystem/Source/Public/LANBeacon.h
+++ b/Engine/Plugins/Online/OnlineSubsystem/Source/Public/LANBeacon.h
@@ -168,7 +168,7 @@ public:
 	/** Port to listen on for LAN queries/responses */
 	int32 LanAnnouncePort;
 
-	/** Unique id to keep UE3 games from seeing each others' LAN packets */
+	/** Unique id to keep UE games from seeing each others' LAN packets */
 	int32 LanGameUniqueId;
 
 	/** Mask containing which platforms can cross communicate */
