Replace/remove uses of whitelist/blacklist in OnlineGroupsInterface

diff --git a/Engine/Plugins/Online/OnlineSubsystem/Source/Public/Interfaces/OnlineGroupsInterface.h b/Engine/Plugins/Online/OnlineSubsystem/Source/Public/Interfaces/OnlineGroupsInterface.h
--- a/Engine/Plugins/Online/OnlineSubsystem/Source/Public/Interfaces/OnlineGroupsInterface.h
+++ b/Engine/Plugins/Online/OnlineSubsystem/Source/Public/Interfaces/OnlineGroupsInterface.h
@@ -107,9 +107,9 @@ struct FGroupMemberRequest
 };
 
 /**
- * An entry in a group blacklist
+ * An entry in a group denylist
  */
-struct FGroupBlacklistEntry
+struct FGroupDenylistEntry
 {
 	FUniqueNetIdPtr GetId() const { return AccountId; }
 
@@ -243,7 +243,7 @@ typedef IGroupUserCollection<FGroupMemberRequest> IGroupRequests;
 /**
  * A structure for caching a list of people who have been banned from a group
  */
-typedef IGroupUserCollection<FGroupBlacklistEntry> IGroupBlacklist;
+typedef IGroupUserCollection<FGroupDenylistEntry> IGroupDenylist;
 
 /**
  * What groups does a particular user currently belong to and what roles do they fill
@@ -629,7 +629,7 @@ public: // can be called by group admins
 	 * @param ContextUserId The ID of the user whose credentials are being used to make this call
 	 * @param GroupId The group's globally unique ID
 	 * @param UserId The user to add to the group.
-	 * @param bAllowBlocked If true, if the user is on the group's blacklist, the user will be unblocked before creating the invite. Otherwise adding a blocked user will fail.
+	 * @param bAllowBlocked If true, if the user is on the group's denylist, the user will be unblocked before creating the invite. Otherwise adding a blocked user will fail.
 	 * @param OnCompleted This callback is invoked after contacting the server. It is guaranteed to occur
 	 *        (regardless of success/fail) and will not be called before this function returns.
 	 */
@@ -696,20 +696,20 @@ public: // can be called by group admins
 	/**
 	 * Ban a user from joining the specified group.
 	 *
-	 * Any user but team owner can be blacklisted including member, invitee, applicant and/or any other user unrelated to the team.
+	 * Any user but team owner can be denylisted including member, invitee, applicant and/or any other user unrelated to the team.
 	 * Blocking follows "ignore" logic: blocked users shouldn't know for sure that they've been blocked.
 	 * Blocked user will still be able to apply for membership and cancel pending application even while blocked.
 	 *
 	 * @param ContextUserId The ID of the user whose credentials are being used to make this call
 	 * @param GroupId The group's globally unique ID
-	 * @param UserId The user to blacklist.
+	 * @param UserId The user to denylist.
 	 * @param OnCompleted This callback is invoked after contacting the server. It is guaranteed to occur
 	 *        (regardless of success/fail) and will not be called before this function returns.
 	 */
 	virtual void BlockUser(const FUniqueNetId& ContextUserId, const FUniqueNetId& GroupId, const FUniqueNetId& UserId, const FOnGroupsRequestCompleted& OnCompleted) = 0;
 
 	/**
-	 * Remove a user from the group's blacklist list.
+	 * Remove a user from the group's denylist list.
 	 *
 	 * @param ContextUserId The ID of the user whose credentials are being used to make this call
 	 * @param GroupId The group's globally unique ID
@@ -765,7 +765,7 @@ public: // can be called by group admins
 
 	/**
 	 * Get the list of users banned from this group.
-	 * When the OnCompleted callback fires, if it succeeded you can use GetCachedGroupBlacklist to retrieve the
+	 * When the OnCompleted callback fires, if it succeeded you can use GetCachedGroupDenylist to retrieve the
 	 * information.
 	 *
 	 * @param ContextUserId The ID of the user whose credentials are being used to make this call
@@ -773,17 +773,17 @@ public: // can be called by group admins
 	 * @param OnCompleted This callback is invoked after contacting the server. It is guaranteed to occur
 	 *        (regardless of success/fail) and will not be called before this function returns.
 	 */
-	virtual void QueryGroupBlacklist(const FUniqueNetId& ContextUserId, const FUniqueNetId& GroupId, const FOnGroupsRequestCompleted& OnCompleted) = 0;
+	virtual void QueryGroupDenylist(const FUniqueNetId& ContextUserId, const FUniqueNetId& GroupId, const FOnGroupsRequestCompleted& OnCompleted) = 0;
 
 	/**
 	 * Get the cached list of users banned from this group. If the information is not cached locally, call
-	 * QueryGroupBlacklist to request it from the server.
+	 * QueryGroupDenylist to request it from the server.
 	 *
 	 * @param ContextUserId The ID of the user whose credentials are being used to make this call
 	 * @param GroupId The group's globally unique ID
 	 * @return Shared pointer to the cached roster structure if one exists
 	 */
-	virtual TSharedPtr<const IGroupBlacklist> GetCachedGroupBlacklist(const FUniqueNetId& ContextUserId, const FUniqueNetId& GroupId) = 0;
+	virtual TSharedPtr<const IGroupDenylist> GetCachedGroupDenylist(const FUniqueNetId& ContextUserId, const FUniqueNetId& GroupId) = 0;
 
 	/**
 	* Queries the server for a list of membership applications that UserId can process (accept, reject, or block) as group admin.
