add new optional analytics delegate for UserCloud impls to notify callers about relevant events

diff --git a/Engine/Plugins/Online/OnlineSubsystem/Source/Public/Interfaces/OnlineUserCloudInterface.h b/Engine/Plugins/Online/OnlineSubsystem/Source/Public/Interfaces/OnlineUserCloudInterface.h
--- a/Engine/Plugins/Online/OnlineSubsystem/Source/Public/Interfaces/OnlineUserCloudInterface.h
+++ b/Engine/Plugins/Online/OnlineSubsystem/Source/Public/Interfaces/OnlineUserCloudInterface.h
@@ -18,6 +18,8 @@ ONLINESUBSYSTEM_API DECLARE_LOG_CATEGORY_EXTERN(LogOnlineCloud, Log, All);
 	UE_CLOG(Conditional, LogOnlineCloud, Verbosity, TEXT("%s%s"), ONLINE_LOG_PREFIX, *FString::Printf(Format, ##__VA_ARGS__)); \
 }
 
+struct FAnalyticsEventAttribute;
+
 /**
  * Delegate fired when the list of files has been returned from the network store
  *
@@ -93,6 +95,15 @@ typedef FOnDeleteUserFileComplete::FDelegate FOnDeleteUserFileCompleteDelegate;
 DECLARE_MULTICAST_DELEGATE_FourParams(FOnRequestUsageInfoComplete, bool, const FUniqueNetId&, int64, const TOptional<int64>&);
 typedef FOnRequestUsageInfoComplete::FDelegate FOnRequestUsageInfoCompleteDelegate;
 
+/**
+ * Delegate fired when we would like to report an event to an analytics provider.
+ *
+ * @param EventName the name of the event
+ * @param Attributes the key/value pairs of analytics event attributes to include with the event
+ */
+DECLARE_MULTICAST_DELEGATE_TwoParams(FOnUserCloudAnalyticsEvent, const FString& /* EventName */, const TArray<FAnalyticsEventAttribute>& /* Attributes */)
+typedef FOnUserCloudAnalyticsEvent::FDelegate FOnUserCloudAnalyticsEventDelegate;
+
 
 /**
  * Provides access to per user cloud file storage
@@ -286,6 +297,14 @@ public:
 	 * @param FileName filename to get state for
 	 */
 	virtual void DumpCloudFileState(const FUniqueNetId& UserId, const FString& FileName) = 0;
+
+	/**
+	 * Delegate fired when we would like to report an event to an analytics provider.
+	 *
+	 * @param EventName the name of the event
+	 * @param Attributes the key/value pairs of analytics event attributes to include with the event
+	 */
+	DEFINE_ONLINE_DELEGATE_TWO_PARAM(OnUserCloudAnalyticsEvent, const FString& /* EventName */, const TArray<FAnalyticsEventAttribute>& /* Attributes */);
 };
 
 typedef TSharedPtr<IOnlineUserCloud, ESPMode::ThreadSafe> IOnlineUserCloudPtr;
