Deprecation of ProvideGameMatchFeedback from public interface

diff --git a/Engine/Plugins/Online/OnlineSubsystem/Source/Public/Interfaces/OnlineGameMatchesInterface.h b/Engine/Plugins/Online/OnlineSubsystem/Source/Public/Interfaces/OnlineGameMatchesInterface.h
--- a/Engine/Plugins/Online/OnlineSubsystem/Source/Public/Interfaces/OnlineGameMatchesInterface.h
+++ b/Engine/Plugins/Online/OnlineSubsystem/Source/Public/Interfaces/OnlineGameMatchesInterface.h
@@ -404,6 +404,7 @@ public:
 	 * @param bReviewTeam - Whether not to review a team or only the members of the team
 	 * @param CompletionDelegate - Completion delegate called when GameMatchFeedback is complete
 	 */
-	virtual void ProvideGameMatchFeedback(const FUniqueNetId& UserId, const FString& MatchId, const bool bReviewTeam, const FOnGameMatchFeedbackComplete& CompletionDelegate) = 0;
+	UE_DEPRECATED(5.0, "ProvideGameMatchFeedback is now deprecated and will be removed.")
+	virtual void ProvideGameMatchFeedback(const FUniqueNetId& UserId, const FString& MatchId, const bool bReviewTeam, const FOnGameMatchFeedbackComplete& CompletionDelegate) {}
 
 };
\ No newline at end of file
