Delete deprecated OnlinePartyInterface APIs
Deprecate more APIs that have replacements that contain more information

diff --git a/Engine/Plugins/Online/OnlineFramework/Source/Party/Private/SocialDebugTools.cpp b/Engine/Plugins/Online/OnlineFramework/Source/Party/Private/SocialDebugTools.cpp
--- a/Engine/Plugins/Online/OnlineFramework/Source/Party/Private/SocialDebugTools.cpp
+++ b/Engine/Plugins/Online/OnlineFramework/Source/Party/Private/SocialDebugTools.cpp
@@ -662,7 +662,7 @@ void USocialDebugTools::FInstanceContext::Init()
 		IOnlinePartyPtr OnlineParty = OnlineSub->GetPartyInterface();
 		if (OnlineParty.IsValid())
 		{
-			PartyInviteReceivedDelegateHandle = OnlineParty->AddOnPartyInviteReceivedDelegate_Handle(FOnPartyInviteReceivedDelegate::CreateUObject(&Owner, &USocialDebugTools::HandlePartyInviteReceived));
+			PartyInviteReceivedDelegateHandle = OnlineParty->AddOnPartyInviteReceivedExDelegate_Handle(FOnPartyInviteReceivedExDelegate::CreateUObject(&Owner, &USocialDebugTools::HandlePartyInviteReceived));
 			PartyJoinRequestReceivedDelegateHandle = OnlineParty->AddOnPartyJoinRequestReceivedDelegate_Handle(FOnPartyJoinRequestReceivedDelegate::CreateUObject(&Owner, &USocialDebugTools::HandlePartyJoinRequestReceived));
 		}
 	}
@@ -709,7 +709,7 @@ void USocialDebugTools::HandleFriendInviteReceived(const FUniqueNetId& LocalUser
 	}
 }
 
-void USocialDebugTools::HandlePartyInviteReceived(const FUniqueNetId& LocalUserId, const FOnlinePartyId& PartyId, const FUniqueNetId& SenderId)
+void USocialDebugTools::HandlePartyInviteReceived(const FUniqueNetId& LocalUserId, const IOnlinePartyJoinInfo& Invitation)
 {
 	if (bAutoAcceptPartyInvites)
 	{
@@ -719,32 +719,23 @@ void USocialDebugTools::HandlePartyInviteReceived(const FUniqueNetId& LocalUserI
 			IOnlinePartyPtr OnlineParty = Context->GetOSS()->GetPartyInterface();
 			if (OnlineParty.IsValid())
 			{
-				TArray<TSharedRef<const IOnlinePartyJoinInfo>> PartyInvites;
-				OnlineParty->GetPendingInvites(LocalUserId, PartyInvites);
-				for (const TSharedRef<const IOnlinePartyJoinInfo>& Invite : PartyInvites)
+				const FString Instance = Context->Name;
+				OnlineParty->JoinParty(LocalUserId, Invitation, FOnJoinPartyComplete::CreateLambda([this, Instance, OnlineParty](const FUniqueNetId& UserId, const FOnlinePartyId& PartyIdTmp, const EJoinPartyCompletionResult Result, const int32 NotApprovedReason)
 				{
-					if (*Invite->GetPartyId() == PartyId)
-					{						
-						const FString Instance = Context->Name;
-						OnlineParty->JoinParty(LocalUserId, *Invite, FOnJoinPartyComplete::CreateLambda([this, Instance, OnlineParty](const FUniqueNetId& UserId, const FOnlinePartyId& PartyIdTmp, const EJoinPartyCompletionResult Result, const int32 NotApprovedReason)
+					bool bSuccess = Result == EJoinPartyCompletionResult::Succeeded;
+					if (bSuccess)
+					{
+						FOnlinePartyDataConstPtr PartyMemberData = GetContext(Instance).GetPartyMemberData();
+						if (PartyMemberData.IsValid())
 						{
-							bool bSuccess = Result == EJoinPartyCompletionResult::Succeeded;
-							if (bSuccess)
-							{
-								FOnlinePartyDataConstPtr PartyMemberData = GetContext(Instance).GetPartyMemberData();
-								if (PartyMemberData.IsValid())
-								{
-									OnlineParty->UpdatePartyMemberData(UserId, PartyIdTmp, DefaultPartyDataNamespace, *PartyMemberData);
-								}
-							}
-							else
-							{
-								UE_LOG(LogParty, Warning, TEXT("Party context[%s] invite join attempt denied for reason [%d]"), *Instance, (int32)Result);
-							}
-						}));
-						break;
+							OnlineParty->UpdatePartyMemberData(UserId, PartyIdTmp, DefaultPartyDataNamespace, *PartyMemberData);
+						}
 					}
-				}
+					else
+					{
+						UE_LOG(LogParty, Warning, TEXT("Party context[%s] invite join attempt denied for reason [%d]"), *Instance, (int32)Result);
+					}
+				}));
 			}
 		}
 	}
diff --git a/Engine/Plugins/Online/OnlineFramework/Source/Party/Public/SocialDebugTools.h b/Engine/Plugins/Online/OnlineFramework/Source/Party/Public/SocialDebugTools.h
--- a/Engine/Plugins/Online/OnlineFramework/Source/Party/Public/SocialDebugTools.h
+++ b/Engine/Plugins/Online/OnlineFramework/Source/Party/Public/SocialDebugTools.h
@@ -100,6 +100,6 @@ private:
 
 	// OSS callback handlers
 	void HandleFriendInviteReceived(const FUniqueNetId& LocalUserId, const FUniqueNetId& FriendId);
-	void HandlePartyInviteReceived(const FUniqueNetId& LocalUserId, const FOnlinePartyId& PartyId, const FUniqueNetId& SenderId);
+	void HandlePartyInviteReceived(const FUniqueNetId& LocalUserId, const IOnlinePartyJoinInfo& Invitation);
 	void HandlePartyJoinRequestReceived(const FUniqueNetId& LocalUserId, const FOnlinePartyId& PartyId, const IOnlinePartyPendingJoinRequestInfo& JoinRequestInfo);
 };
\ No newline at end of file
diff --git a/Engine/Plugins/Online/OnlineSubsystem/Source/Private/OnlinePartyInterface.cpp b/Engine/Plugins/Online/OnlineSubsystem/Source/Private/OnlinePartyInterface.cpp
--- a/Engine/Plugins/Online/OnlineSubsystem/Source/Private/OnlinePartyInterface.cpp
+++ b/Engine/Plugins/Online/OnlineSubsystem/Source/Private/OnlinePartyInterface.cpp
@@ -140,42 +140,31 @@ FString FPartyInvitationRecipient::ToDebugString() const
 	return FString::Printf(TEXT("Id=[%s], PlatformData=[%s]"), *Id->ToDebugString(), *PlatformData);
 }
 
-FDelegateHandle IOnlinePartySystem::AddOnPartyJIPDelegate_Handle(const FOnPartyJIPDelegate& Delegate)
+PRAGMA_DISABLE_DEPRECATION_WARNINGS
+FDelegateHandle IOnlinePartySystem::AddOnPartyInviteReceivedDelegate_Handle(const FOnPartyInviteReceivedDelegate& Delegate)
+PRAGMA_ENABLE_DEPRECATION_WARNINGS
 {
-	auto DeprecationHelperLambda = [Delegate](const FUniqueNetId& LocalUserId, const FOnlinePartyId& PartyId, bool Success, int32 /*DeniedResultCode*/)
+	auto DeprecationHelperLambda = [Delegate](const FUniqueNetId& LocalUserId, const IOnlinePartyJoinInfo& Invitation)
 	{
-		Delegate.ExecuteIfBound(LocalUserId, PartyId, Success);
+		Delegate.ExecuteIfBound(LocalUserId, *Invitation.GetPartyId(), *Invitation.GetSourceUserId());
 	};
-	return OnPartyJIPResponseDelegates.Add(FOnPartyJIPResponseDelegate::CreateLambda(DeprecationHelperLambda));
+	return OnPartyInviteReceivedExDelegates.Add(FOnPartyInviteReceivedExDelegate::CreateLambda(DeprecationHelperLambda));
 }
 
-FDelegateHandle IOnlinePartySystem::AddOnPartyDataReceivedDelegate_Handle(const FOnPartyDataReceivedConstDelegate& Delegate)
+PRAGMA_DISABLE_DEPRECATION_WARNINGS
+FDelegateHandle IOnlinePartySystem::AddOnPartyInviteRemovedDelegate_Handle(const FOnPartyInviteRemovedDelegate& Delegate)
+PRAGMA_ENABLE_DEPRECATION_WARNINGS
 {
-	auto DeprecationHelperLambda = [Delegate](const FUniqueNetId& LocalUserId, const FOnlinePartyId& PartyId, const FName& Namespace, const FOnlinePartyData& PartyData)
+	auto DeprecationHelperLambda = [Delegate](const FUniqueNetId& LocalUserId, const IOnlinePartyJoinInfo& Invitation, EPartyInvitationRemovedReason Reason)
 	{
-		// For backwards compatibility, only forward the default namespace
-		if (Namespace == DefaultPartyDataNamespace)
-		{
-			Delegate.ExecuteIfBound(LocalUserId, PartyId, PartyData);
-		}
-	};
-	return OnPartyDataReceivedDelegates.Add(FOnPartyDataReceivedDelegate::CreateLambda(DeprecationHelperLambda));
-}
-
-FDelegateHandle IOnlinePartySystem::AddOnPartyMemberDataReceivedDelegate_Handle(const FOnPartyMemberDataReceivedConstDelegate& Delegate)
-{
-	auto DeprecationHelperLambda = [Delegate](const FUniqueNetId& LocalUserId, const FOnlinePartyId& PartyId, const FUniqueNetId& MemberId, const FName& Namespace, const FOnlinePartyData& PartyData)
-	{
-		// For backwards compatibility, only forward the default namespace
-		if (Namespace == DefaultPartyDataNamespace)
-		{
-			Delegate.ExecuteIfBound(LocalUserId, PartyId, MemberId, PartyData);
-		}
+		Delegate.ExecuteIfBound(LocalUserId, *Invitation.GetPartyId(), *Invitation.GetSourceUserId(), Reason);
 	};
-	return OnPartyMemberDataReceivedDelegates.Add(FOnPartyMemberDataReceivedDelegate::CreateLambda(DeprecationHelperLambda));
+	return OnPartyInviteRemovedExDelegates.Add(FOnPartyInviteRemovedExDelegate::CreateLambda(DeprecationHelperLambda));
 }
 
+PRAGMA_DISABLE_DEPRECATION_WARNINGS
 void IOnlinePartySystem::QueryPartyJoinability(const FUniqueNetId& LocalUserId, const IOnlinePartyJoinInfo& OnlinePartyJoinInfo, const FOnQueryPartyJoinabilityComplete& Delegate)
+PRAGMA_ENABLE_DEPRECATION_WARNINGS
 {
 	// Default implementation will call the ex version.
 	QueryPartyJoinability(LocalUserId, OnlinePartyJoinInfo, FOnQueryPartyJoinabilityCompleteEx::CreateLambda([Delegate](const FUniqueNetId& LambdaLocalUserId, const FOnlinePartyId& LambdaPartyId, const FQueryPartyJoinabilityResult& QueryPartyJoinabilityResult)
@@ -184,11 +173,6 @@ void IOnlinePartySystem::QueryPartyJoinability(const FUniqueNetId& LocalUserId,
 	}));
 }
 
-void IOnlinePartySystem::RespondToQueryJoinability(const FUniqueNetId& LocalUserId, const FOnlinePartyId& PartyId, const FUniqueNetId& RecipientId, bool bCanJoin, int32 DeniedResultCode)
-{
-	RespondToQueryJoinability(LocalUserId, PartyId, RecipientId, bCanJoin, DeniedResultCode, FOnlinePartyDataConstPtr());
-}
-
 bool FPartyConfiguration::operator==(const FPartyConfiguration& Other) const
 {
 	return JoinRequestAction == Other.JoinRequestAction &&
diff --git a/Engine/Plugins/Online/OnlineSubsystem/Source/Public/Interfaces/OnlinePartyInterface.h b/Engine/Plugins/Online/OnlineSubsystem/Source/Public/Interfaces/OnlinePartyInterface.h
--- a/Engine/Plugins/Online/OnlineSubsystem/Source/Public/Interfaces/OnlinePartyInterface.h
+++ b/Engine/Plugins/Online/OnlineSubsystem/Source/Public/Interfaces/OnlinePartyInterface.h
@@ -860,6 +860,7 @@ DECLARE_DELEGATE_FourParams(FOnJoinPartyComplete, const FUniqueNetId& /*LocalUse
  * @param Result - result of the operation
  * @param NotApprovedReason - client defined value describing why you were not approved
  */
+UE_DEPRECATED(5.1, "Use FOnQueryPartyJoinabilityCompleteEx")
 DECLARE_DELEGATE_FourParams(FOnQueryPartyJoinabilityComplete, const FUniqueNetId& /*LocalUserId*/, const FOnlinePartyId& /*PartyId*/, const EJoinPartyCompletionResult /*Result*/, const int32 /*NotApprovedReason*/);
 /**
  * Party query joinability async task completed callback
@@ -989,7 +990,8 @@ PARTY_DECLARE_DELEGATETYPE(OnPartyStateChanged);
 * @param DeniedResultCode - descriptive reason for a failure to join
 */
 DECLARE_MULTICAST_DELEGATE_ThreeParams(F_PREFIX(OnPartyJIP), const FUniqueNetId& /*LocalUserId*/, const FOnlinePartyId& /*PartyId*/, bool /*Success*/);
-PARTY_DECLARE_DELEGATETYPE(OnPartyJIP);
+UE_DEPRECATED(5.1, "Use FOnPartyJIPResponseDelegate instead of FOnPartyJIPDelegate")
+typedef FOnPartyJIP::FDelegate FOnPartyJIPDelegate;
 DECLARE_MULTICAST_DELEGATE_FourParams(F_PREFIX(OnPartyJIPResponse), const FUniqueNetId& /*LocalUserId*/, const FOnlinePartyId& /*PartyId*/, bool /*Success*/, int32 /*DeniedResultCode*/);
 PARTY_DECLARE_DELEGATETYPE(OnPartyJIPResponse);
 
@@ -1011,9 +1013,6 @@ PARTY_DECLARE_DELEGATETYPE(OnPartyPromotionLockoutChanged);
 DECLARE_MULTICAST_DELEGATE_ThreeParams(F_PREFIX(OnPartyConfigChanged), const FUniqueNetId& /*LocalUserId*/, const FOnlinePartyId& /*PartyId*/, const FPartyConfiguration& /*PartyConfig*/);
 PARTY_DECLARE_DELEGATETYPE(OnPartyConfigChanged);
 
-UE_DEPRECATED(4.26, "Use OnPartyConfigChangedDelegate instead of OnPartyConfigChangedConstDelegate")
-typedef FOnPartyConfigChangedDelegate FOnPartyConfigChangedConstDelegate;
-
 /**
  * Notification when party data is updated
  * Deprecated - Use OnPartyDataReceived
@@ -1022,7 +1021,8 @@ typedef FOnPartyConfigChangedDelegate FOnPartyConfigChangedConstDelegate;
  * @param PartyData - party data that was updated
  */
 DECLARE_MULTICAST_DELEGATE_ThreeParams(F_PREFIX(OnPartyDataReceivedConst), const FUniqueNetId& /*LocalUserId*/, const FOnlinePartyId& /*PartyId*/, const FOnlinePartyData& /*PartyData*/);
-PARTY_DECLARE_DELEGATETYPE(OnPartyDataReceivedConst);
+UE_DEPRECATED(5.1, "Use OnPartyDataReceived instead of OnPartyDataReceivedConst")
+typedef FOnPartyDataReceivedConst::FDelegate FOnPartyDataReceivedConstDelegate;
 
 /**
  * Notification when party data is updated
@@ -1071,7 +1071,8 @@ PARTY_DECLARE_DELEGATETYPE(OnPartyMemberJoined);
  * @param PartyMemberData - party member data that was updated
  */
 DECLARE_MULTICAST_DELEGATE_FourParams(F_PREFIX(OnPartyMemberDataReceivedConst), const FUniqueNetId& /*LocalUserId*/, const FOnlinePartyId& /*PartyId*/, const FUniqueNetId& /*MemberId*/, const FOnlinePartyData& /*PartyMemberData*/);
-PARTY_DECLARE_DELEGATETYPE(OnPartyMemberDataReceivedConst);
+UE_DEPRECATED(5.1, "Use OnPartyMemberDataReceived instead of OnPartyMemberDataReceivedConst")
+typedef FOnPartyMemberDataReceivedConst::FDelegate FOnPartyMemberDataReceivedConstDelegate;
 
 /**
  * Notification when party member data is updated
@@ -1109,7 +1110,8 @@ PARTY_DECLARE_DELEGATETYPE(OnPartyInviteRequestReceived);
  * @param SenderId id of member that sent the invite
  */
 DECLARE_MULTICAST_DELEGATE_ThreeParams(F_PREFIX(OnPartyInviteReceived), const FUniqueNetId& /*LocalUserId*/, const FOnlinePartyId& /*PartyId*/, const FUniqueNetId& /*SenderId*/);
-PARTY_DECLARE_DELEGATETYPE(OnPartyInviteReceived);
+UE_DEPRECATED(5.1, "Use OnPartyInviteReceivedEx instead of OnPartyInviteReceived")
+typedef FOnPartyInviteReceived::FDelegate FOnPartyInviteReceivedDelegate;
 
 /**
  * Notification when a new invite is received
@@ -1128,7 +1130,8 @@ PARTY_DECLARE_DELEGATETYPE(OnPartyInviteReceivedEx);
  * @param Reason reason the invite has been removed
  */
 DECLARE_MULTICAST_DELEGATE_FourParams(F_PREFIX(OnPartyInviteRemoved), const FUniqueNetId& /*LocalUserId*/, const FOnlinePartyId& /*PartyId*/, const FUniqueNetId& /*SenderId*/, EPartyInvitationRemovedReason /*Reason*/);
-PARTY_DECLARE_DELEGATETYPE(OnPartyInviteRemoved);
+UE_DEPRECATED(5.1, "Use OnPartyInviteRemovedEx instead of OnPartyInviteRemoved")
+typedef FOnPartyInviteRemoved::FDelegate FOnPartyInviteRemovedDelegate;
 
 /**
  * Notification when an invite has been removed
@@ -1158,9 +1161,6 @@ PARTY_DECLARE_DELEGATETYPE(OnPartyInviteResponseReceived);
 DECLARE_MULTICAST_DELEGATE_ThreeParams(F_PREFIX(OnPartyJoinRequestReceived), const FUniqueNetId& /*LocalUserId*/, const FOnlinePartyId& /*PartyId*/, const IOnlinePartyPendingJoinRequestInfo& /*JoinRequestInfo*/);
 PARTY_DECLARE_DELEGATETYPE(OnPartyJoinRequestReceived);
 
-UE_DEPRECATED(4.26, "Use OnPartyJoinRequestReceived instead of OnPartyGroupJoinRequestReceived")
-typedef FOnPartyJoinRequestReceivedDelegate FOnPartyGroupJoinRequestReceivedDelegate;
-
 /**
  * Notification when a new reservation request is received
  * @param LocalUserId - id associated with this notification
@@ -1181,9 +1181,6 @@ PARTY_DECLARE_DELEGATETYPE(OnPartyJIPRequestReceived);
 DECLARE_MULTICAST_DELEGATE_ThreeParams(F_PREFIX(OnQueryPartyJoinabilityReceived), const FUniqueNetId& /*LocalUserId*/, const FOnlinePartyId& /*PartyId*/, const IOnlinePartyPendingJoinRequestInfo& /*JoinRequestInfo*/);
 PARTY_DECLARE_DELEGATETYPE(OnQueryPartyJoinabilityReceived);
 
-UE_DEPRECATED(4.26, "Use FOnQueryPartyJoinabilityReceivedDelegate instead of OnQueryPartyJoinabilityGroupReceived")
-typedef FOnQueryPartyJoinabilityReceivedDelegate FOnQueryPartyJoinabilityGroupReceivedDelegate;
-
 /**
  * Request for the game to fill in data to be sent with the join request for the leader to make an informed decision based on the joiner's state
  * @param LocalUserId - id associated with this notification
@@ -1230,6 +1227,18 @@ PARTY_DECLARE_DELEGATETYPE(OnPartyRequestToJoinReceived);
 DECLARE_MULTICAST_DELEGATE_FiveParams(F_PREFIX(OnPartyRequestToJoinRemoved), const FUniqueNetId& /*LocalUserId*/, const FOnlinePartyId& /*PartyId*/, const FUniqueNetId& /*RequesterId*/, const IOnlinePartyRequestToJoinInfo& /*Request*/, EPartyRequestToJoinRemovedReason /*Reason*/);
 PARTY_DECLARE_DELEGATETYPE(OnPartyRequestToJoinRemoved);
 
+#define DEFINE_DEPRECATED_PARTY_DELEGATE(DeprecatedDelegateName, DelegateName) \
+	virtual void Clear##DeprecatedDelegateName##Delegate_Handle(FDelegateHandle& Handle) \
+	{ \
+		DelegateName##Delegates.Remove(Handle); \
+		Handle.Reset(); \
+	} \
+	virtual void Clear##DeprecatedDelegateName##Delegates(void* Object) \
+	{ \
+		DelegateName##Delegates.RemoveAll(Object); \
+	}
+
+
 /**
  * Interface definition for the online party services 
  * Allows for forming a party and communicating with party members
@@ -1340,7 +1349,10 @@ public:
 	 * @param OnlinePartyJoinInfo - join information containing data such as party id, leader id
 	 * @param Delegate - called on completion
 	 */
+	PRAGMA_DISABLE_DEPRECATION_WARNINGS // Suppress deprecation warning for FOnQueryPartyJoinabilityComplete
+	UE_DEPRECATED(5.1, "Use QueryPartyJoinability with the FOnQueryPartyJoinabilityCompleteEx delegate instead")
 	virtual void QueryPartyJoinability(const FUniqueNetId& LocalUserId, const IOnlinePartyJoinInfo& OnlinePartyJoinInfo, const FOnQueryPartyJoinabilityComplete& Delegate);
+	PRAGMA_ENABLE_DEPRECATION_WARNINGS
 
 	/**
 	 * Query a party to check it's current joinability
@@ -1428,8 +1440,6 @@ public:
 	 * @param PartyData - data to send back to the querying user
 	 */
 	virtual void RespondToQueryJoinability(const FUniqueNetId& LocalUserId, const FOnlinePartyId& PartyId, const FUniqueNetId& RecipientId, bool bCanJoin, int32 DeniedResultCode, FOnlinePartyDataConstPtr PartyData) = 0;
-	UE_DEPRECATED(4.26, "Use RespondToQueryJoinability with PartyData")
-	virtual void RespondToQueryJoinability(const FUniqueNetId& LocalUserId, const FOnlinePartyId& PartyId, const FUniqueNetId& RecipientId, bool bCanJoin, int32 DeniedResultCode = 0);
 
 	/**
 	 * sends an invitation to a user that could not otherwise join a party
@@ -1491,21 +1501,6 @@ public:
 	 */
 	virtual bool PromoteMember(const FUniqueNetId& LocalUserId, const FOnlinePartyId& PartyId, const FUniqueNetId& TargetMemberId, const FOnPromotePartyMemberComplete& Delegate = FOnPromotePartyMemberComplete()) = 0;
 
-	/**
-	 * Set party data and broadcast to all members
-	 * Only current data can be set and no history of past party data is preserved
-	 * Party members notified of new data (see FOnPartyDataReceived)
-	 *
-	 * @param LocalUserId - user making the request
-	 * @param PartyId - id of an existing party
-	 * @param PartyData - data to send to all party members
-	 * @param Delegate - called on completion
-	 *
-	 * @return true if task was started
-	 */
-	UE_DEPRECATED(4.26, "Party data is separated by namespace. Use UpdatePartyData that has a Namespace parameter.")
-	virtual bool UpdatePartyData(const FUniqueNetId& LocalUserId, const FOnlinePartyId& PartyId, const FOnlinePartyData& PartyData) { return UpdatePartyData(LocalUserId, PartyId, DefaultPartyDataNamespace, PartyData); }
-
 	/**
 	 * Set party data and broadcast to all members
 	 * Only current data can be set and no history of past party data is preserved
@@ -1521,21 +1516,6 @@ public:
 	 */
 	virtual bool UpdatePartyData(const FUniqueNetId& LocalUserId, const FOnlinePartyId& PartyId, const FName& Namespace, const FOnlinePartyData& PartyData) = 0;
 
-	/**
-	 * Set party data for a single party member and broadcast to all members
-	 * Only current data can be set and no history of past party member data is preserved
-	 * Party members notified of new data (see FOnPartyMemberDataReceived)
-	 *
-	 * @param LocalUserId - user making the request
-	 * @param PartyId - id of an existing party
-	 * @param PartyMemberData - member data to send to all party members
-	 * @param Delegate - called on completion
-	 *
-	 * @return true if task was started
-	 */
-	UE_DEPRECATED(4.26, "Party data is separated by namespace. Use UpdatePartyMemberData that has a Namespace parameter.")
-	virtual bool UpdatePartyMemberData(const FUniqueNetId& LocalUserId, const FOnlinePartyId& PartyId, const FOnlinePartyData& PartyMemberData)  { return UpdatePartyMemberData(LocalUserId, PartyId, DefaultPartyDataNamespace, PartyMemberData); }
-
 	/**
 	 * Set party data for a single party member and broadcast to all members
 	 * Only current data can be set and no history of past party member data is preserved
@@ -1601,18 +1581,6 @@ public:
 	 */
 	virtual FOnlinePartyMemberConstPtr GetPartyMember(const FUniqueNetId& LocalUserId, const FOnlinePartyId& PartyId, const FUniqueNetId& MemberId) const = 0;
 
-	/**
-	 * Get current cached data associated with a party
-	 * FOnPartyDataReceived notification called whenever this data changes
-	 *
-	 * @param LocalUserId - user making the request
-	 * @param PartyId     - id of an existing party
-	 *
-	 * @return party data or nullptr if not found
-	 */
-	UE_DEPRECATED(4.26, "Party data is separated by namespace. Use GetPartyData that has a Namespace parameter.")
-	virtual FOnlinePartyDataConstPtr GetPartyData(const FUniqueNetId& LocalUserId, const FOnlinePartyId& PartyId) const { return GetPartyData(LocalUserId, PartyId, DefaultPartyDataNamespace); }
-
 	/**
 	 * Get current cached data associated with a party
 	 * FOnPartyDataReceived notification called whenever this data changes
@@ -1625,19 +1593,6 @@ public:
 	 */
 	virtual FOnlinePartyDataConstPtr GetPartyData(const FUniqueNetId& LocalUserId, const FOnlinePartyId& PartyId, const FName& Namespace) const = 0;
 
-	/**
-	 * Get current cached data associated with a party member
-	 * FOnPartyMemberDataReceived notification called whenever this data changes
-	 *
-	 * @param LocalUserId - user making the request
-	 * @param PartyId     - id of an existing party
-	 * @param MemberId    - id of member to find data for
-	 *
-	 * @return party member data or nullptr if not found
-	 */
-	UE_DEPRECATED(4.26, "Party data is separated by namespace. Use GetPartyMemberData that has a Namespace parameter.")
-	virtual FOnlinePartyDataConstPtr GetPartyMemberData(const FUniqueNetId& LocalUserId, const FOnlinePartyId& PartyId, const FUniqueNetId& MemberId) const { return GetPartyMemberData(LocalUserId, PartyId, MemberId, DefaultPartyDataNamespace); }
-
 	/**
 	 * Get current cached data associated with a party member
 	 * FOnPartyMemberDataReceived notification called whenever this data changes
@@ -1793,9 +1748,7 @@ public:
 	 * OnPartyMemberDataReceived
 	 * OnPartyInvitesChanged
 	 * OnPartyInviteRequestReceived
-	 * OnPartyInviteReceived
 	 * OnPartyInviteReceivedEx
-	 * OnPartyInviteRemoved
 	 * OnPartyInviteRemovedEx
 	 * OnPartyInviteResponseReceived
 	 * OnPartyJoinRequestReceived
@@ -1838,8 +1791,6 @@ public:
 	* @param DeniedResultCode - descriptive reason for a failure to join
 	*/
 	DEFINE_ONLINE_DELEGATE_FOUR_PARAM(OnPartyJIPResponse, const FUniqueNetId& /*LocalUserId*/, const FOnlinePartyId& /*PartyId*/, bool /*Success*/, int32 /*DeniedResultCode*/);
-	UE_DEPRECATED(5.0, "Use OnPartyJIPResponse instead of OnPartyJIP")
-	virtual FDelegateHandle AddOnPartyJIPDelegate_Handle(const FOnPartyJIPDelegate& Delegate);
 
 	/**
 	 * Notification when player promotion is locked out.
@@ -1865,8 +1816,6 @@ public:
 	 * @param PartyData - party data that was updated
 	 */
 	DEFINE_ONLINE_DELEGATE_FOUR_PARAM(OnPartyDataReceived, const FUniqueNetId& /*LocalUserId*/, const FOnlinePartyId& /*PartyId*/, const FName& /*Namespace*/, const FOnlinePartyData& /*PartyData*/);
-	UE_DEPRECATED(4.26, "Use OnPartyDataReceived instead of OnPartyDataReceivedConst")
-	virtual FDelegateHandle AddOnPartyDataReceivedDelegate_Handle(const FOnPartyDataReceivedConstDelegate& Delegate); // Deprecated method that only forwards PartyData for the default namespace
 
 	/**
 	* Notification when a member is promoted in a party
@@ -1902,8 +1851,6 @@ public:
 	 * @param PartyMemberData - party member data that was updated
 	 */
 	DEFINE_ONLINE_DELEGATE_FIVE_PARAM(OnPartyMemberDataReceived, const FUniqueNetId& /*LocalUserId*/, const FOnlinePartyId& /*PartyId*/, const FUniqueNetId& /*MemberId*/, const FName& /*Namespace*/, const FOnlinePartyData& /*PartyData*/);
-	UE_DEPRECATED(4.26, "Use OnPartyMemberDataReceived instead of OnPartyMemberDataReceivedConst")
-	virtual FDelegateHandle AddOnPartyMemberDataReceivedDelegate_Handle(const FOnPartyMemberDataReceivedConstDelegate& Delegate);
 
 	/**
 	 * Notification when an invite list has changed for a party
@@ -1920,31 +1867,17 @@ public:
 	*/
 	DEFINE_ONLINE_DELEGATE_FOUR_PARAM(OnPartyInviteRequestReceived, const FUniqueNetId& /*LocalUserId*/, const FOnlinePartyId& /*PartyId*/, const FUniqueNetId& /*SenderId*/, const FUniqueNetId& /*RequestForId*/);
 
-	/**
-	 * Notification when a new invite is received
-	 * Deprecated - use OnPartyInviteReceivedEx
-	 * @param LocalUserId - id associated with this notification
-	 * @param PartyId - id associated with the party
-	 * @param SenderId - id of member that sent the invite
-	 */
-	DEFINE_ONLINE_DELEGATE_THREE_PARAM(OnPartyInviteReceived, const FUniqueNetId& /*LocalUserId*/, const FOnlinePartyId& /*PartyId*/, const FUniqueNetId& /*SenderId*/);
-
 	/**
 	 * Notification when a new invite is received
 	 * @param LocalUserId - id associated with this notification
 	 * @param Invitation - the invitation that was received
 	 */
 	DEFINE_ONLINE_DELEGATE_TWO_PARAM(OnPartyInviteReceivedEx, const FUniqueNetId& /*LocalUserId*/, const IOnlinePartyJoinInfo& /*Invitation*/);
-
-	/**
-	 * Notification when an invite has been removed
-	 * Deprecated Use OnPartyInviteRemovedEx
-	 * @param LocalUserId id associated with this notification
-	 * @param PartyId id associated with the party
-	 * @param SenderId id of member that sent the invite
-	 * @param Reason reason the invitation was removed
-	 */
-	DEFINE_ONLINE_DELEGATE_FOUR_PARAM(OnPartyInviteRemoved, const FUniqueNetId& /*LocalUserId*/, const FOnlinePartyId& /*PartyId*/, const FUniqueNetId& /*SenderId*/, EPartyInvitationRemovedReason /*Reason*/);
+	PRAGMA_DISABLE_DEPRECATION_WARNINGS // Suppress deprecation warning for FOnPartyInviteReceivedDelegate
+	UE_DEPRECATED(5.1, "Use OnPartyInviteReceivedEx instead of OnPartyInviteReceived")
+	virtual FDelegateHandle AddOnPartyInviteReceivedDelegate_Handle(const FOnPartyInviteReceivedDelegate& Delegate);
+	DEFINE_DEPRECATED_PARTY_DELEGATE(OnPartyInviteReceived, OnPartyInviteReceivedEx);
+	PRAGMA_ENABLE_DEPRECATION_WARNINGS
 
 	/**
 	 * Notification when an invite has been removed
@@ -1953,6 +1886,11 @@ public:
 	 * @param Reason reason the invitation was removed
 	 */
 	DEFINE_ONLINE_DELEGATE_THREE_PARAM(OnPartyInviteRemovedEx, const FUniqueNetId& /*LocalUserId*/, const IOnlinePartyJoinInfo& /*Invitation*/, EPartyInvitationRemovedReason /*Reason*/);
+	PRAGMA_DISABLE_DEPRECATION_WARNINGS // Suppress deprecation warning for FOnPartyInviteRemovedDelegate
+	UE_DEPRECATED(5.1, "Use OnPartyInviteRemovedEx instead of OnPartyInviteRemoved")
+	virtual FDelegateHandle AddOnPartyInviteRemovedDelegate_Handle(const FOnPartyInviteRemovedDelegate& Delegate);
+	DEFINE_DEPRECATED_PARTY_DELEGATE(OnPartyInviteRemoved, OnPartyInviteRemovedEx);
+	PRAGMA_ENABLE_DEPRECATION_WARNINGS
 
 	/**
 	 * Notification when an invitation response is received
