[Engine/Online]  Adding support for a sender to cancel a party join invitation.

diff --git a/Engine/Plugins/Online/OnlineSubsystem/Source/Private/OnlinePartyInterface.cpp b/Engine/Plugins/Online/OnlineSubsystem/Source/Private/OnlinePartyInterface.cpp
--- a/Engine/Plugins/Online/OnlineSubsystem/Source/Private/OnlinePartyInterface.cpp
+++ b/Engine/Plugins/Online/OnlineSubsystem/Source/Private/OnlinePartyInterface.cpp
@@ -425,6 +425,30 @@ const TCHAR* ToString(const ESendPartyInvitationCompletionResult Value)
 	return TEXT("Unknown");
 }
 
+const TCHAR* ToString(const ECancelPartyInvitationCompletionResult Value)
+{
+	switch (Value)
+	{
+	case ECancelPartyInvitationCompletionResult::InviteNotFound:
+	{
+		return TEXT("InviteNotFound");
+	}
+	case ECancelPartyInvitationCompletionResult::InvalidParty:
+	{
+		return TEXT("InvalidParty");
+	}
+	case ECancelPartyInvitationCompletionResult::UnknownInternalFailure:
+	{
+		return TEXT("UnknownInternalFailure");
+	}
+	case ECancelPartyInvitationCompletionResult::Succeeded:
+	{
+		return TEXT("Succeeded");
+	}
+	}
+	return TEXT("Unknown");
+}
+
 const TCHAR* ToString(const EJoinPartyCompletionResult Value)
 {
 	switch (Value)
