[Engine/Online]  Adding support for a sender to cancel a party join invitation.

diff --git a/Engine/Plugins/Online/OnlineSubsystem/Source/Public/Interfaces/OnlinePartyInterface.h b/Engine/Plugins/Online/OnlineSubsystem/Source/Public/Interfaces/OnlinePartyInterface.h
--- a/Engine/Plugins/Online/OnlineSubsystem/Source/Public/Interfaces/OnlinePartyInterface.h
+++ b/Engine/Plugins/Online/OnlineSubsystem/Source/Public/Interfaces/OnlinePartyInterface.h
@@ -37,6 +37,7 @@ enum class EPromoteMemberCompletionResult : int8;
 enum class ERejectPartyInvitationCompletionResult : int8;
 enum class ERequestPartyInvitationCompletionResult : int8;
 enum class ESendPartyInvitationCompletionResult : int8;
+enum class ECancelPartyInvitationCompletionResult : int8;
 enum class EUpdateConfigCompletionResult : int8;
 enum class EInvitationResponse : uint8;
 enum class ERequestToJoinPartyCompletionResult : int8;
@@ -903,6 +904,15 @@ DECLARE_DELEGATE_ThreeParams(FOnRequestPartyInvitationComplete, const FUniqueNet
  * @param Result - result of the send operation
  */
 DECLARE_DELEGATE_FourParams(FOnSendPartyInvitationComplete, const FUniqueNetId& /*LocalUserId*/, const FOnlinePartyId& /*PartyId*/, const FUniqueNetId& /*RecipientId*/, const ESendPartyInvitationCompletionResult /*Result*/);
+/**
+ * Party invitation cancel completed callback
+ *
+ * @param SenderUserId - id of user that initiated the request
+ * @param PartyId - id associated with the party
+ * @param TargetUserId - user invite was sent to
+ * @param Result - result of the cancel operation
+ */
+DECLARE_DELEGATE_FourParams(FOnCancelPartyInvitationComplete, const FUniqueNetId& /*SenderUserId*/, const FOnlinePartyId& /*PartyId*/, const FUniqueNetId& /*TargetUserId*/, const ECancelPartyInvitationCompletionResult /*Result*/);
 /**
  * Accepting an invite to a user to join party async task completed callback
  *
@@ -1454,6 +1464,16 @@ public:
 	 */
 	virtual bool SendInvitation(const FUniqueNetId& LocalUserId, const FOnlinePartyId& PartyId, const FPartyInvitationRecipient& Recipient, const FOnSendPartyInvitationComplete& Delegate = FOnSendPartyInvitationComplete()) = 0;
 
+	/**
+	 * Cancel an invitation sent to a user
+	 *
+	 * @param LocalUserId - user making the cancellation
+	 * @param TargetUserId - previously invited user
+	 * @param PartyId - id of an existing party
+	 * @param Delegate - called on completion
+	 */
+	virtual void CancelInvitation(const FUniqueNetId& LocalUserId, const FUniqueNetId& TargetUserId, const FOnlinePartyId& PartyId, const FOnCancelPartyInvitationComplete& Delegate = FOnCancelPartyInvitationComplete()) = 0;
+
 	/**
 	 * Reject an invite to a party
 	 *
@@ -2112,6 +2132,14 @@ enum class ESendPartyInvitationCompletionResult : int8
 	Succeeded = 1
 };
 
+enum class ECancelPartyInvitationCompletionResult : int8
+{
+	InviteNotFound = -100,
+	InvalidParty,
+	UnknownInternalFailure = 0,
+	Succeeded = 1
+};
+
 enum class EAcceptPartyInvitationCompletionResult : int8
 {
 	NotLoggedIn = -100,
@@ -2216,6 +2244,8 @@ ONLINESUBSYSTEM_API const TCHAR* ToString(const ECreatePartyCompletionResult Val
 /** @return the stringified version of the enum passed in */
 ONLINESUBSYSTEM_API const TCHAR* ToString(const ESendPartyInvitationCompletionResult Value);
 /** @return the stringified version of the enum passed in */
+ONLINESUBSYSTEM_API const TCHAR* ToString(const ECancelPartyInvitationCompletionResult Value);
+/** @return the stringified version of the enum passed in */
 ONLINESUBSYSTEM_API const TCHAR* ToString(const EJoinPartyCompletionResult Value);
 /** @return the stringified version of the enum passed in */
 ONLINESUBSYSTEM_API const TCHAR* ToString(const ELeavePartyCompletionResult Value);
