[Engine/Online]  Changing  IOnlinePartySystem::CancelInvitation to use new FOnlineError return value framework

diff --git a/Engine/Plugins/Online/OnlineSubsystem/Source/Private/OnlinePartyInterface.cpp b/Engine/Plugins/Online/OnlineSubsystem/Source/Private/OnlinePartyInterface.cpp
--- a/Engine/Plugins/Online/OnlineSubsystem/Source/Private/OnlinePartyInterface.cpp
+++ b/Engine/Plugins/Online/OnlineSubsystem/Source/Private/OnlinePartyInterface.cpp
@@ -425,30 +425,6 @@ const TCHAR* ToString(const ESendPartyInvitationCompletionResult Value)
 	return TEXT("Unknown");
 }
 
-const TCHAR* ToString(const ECancelPartyInvitationCompletionResult Value)
-{
-	switch (Value)
-	{
-	case ECancelPartyInvitationCompletionResult::InviteNotFound:
-	{
-		return TEXT("InviteNotFound");
-	}
-	case ECancelPartyInvitationCompletionResult::InvalidParty:
-	{
-		return TEXT("InvalidParty");
-	}
-	case ECancelPartyInvitationCompletionResult::UnknownInternalFailure:
-	{
-		return TEXT("UnknownInternalFailure");
-	}
-	case ECancelPartyInvitationCompletionResult::Succeeded:
-	{
-		return TEXT("Succeeded");
-	}
-	}
-	return TEXT("Unknown");
-}
-
 const TCHAR* ToString(const EJoinPartyCompletionResult Value)
 {
 	switch (Value)
diff --git a/Engine/Plugins/Online/OnlineSubsystem/Source/Public/Interfaces/OnlinePartyInterface.h b/Engine/Plugins/Online/OnlineSubsystem/Source/Public/Interfaces/OnlinePartyInterface.h
--- a/Engine/Plugins/Online/OnlineSubsystem/Source/Public/Interfaces/OnlinePartyInterface.h
+++ b/Engine/Plugins/Online/OnlineSubsystem/Source/Public/Interfaces/OnlinePartyInterface.h
@@ -37,7 +37,6 @@ enum class EPromoteMemberCompletionResult : int8;
 enum class ERejectPartyInvitationCompletionResult : int8;
 enum class ERequestPartyInvitationCompletionResult : int8;
 enum class ESendPartyInvitationCompletionResult : int8;
-enum class ECancelPartyInvitationCompletionResult : int8;
 enum class EUpdateConfigCompletionResult : int8;
 enum class EInvitationResponse : uint8;
 enum class ERequestToJoinPartyCompletionResult : int8;
@@ -912,7 +911,7 @@ DECLARE_DELEGATE_FourParams(FOnSendPartyInvitationComplete, const FUniqueNetId&
  * @param TargetUserId - user invite was sent to
  * @param Result - result of the cancel operation
  */
-DECLARE_DELEGATE_FourParams(FOnCancelPartyInvitationComplete, const FUniqueNetId& /*SenderUserId*/, const FOnlinePartyId& /*PartyId*/, const FUniqueNetId& /*TargetUserId*/, const ECancelPartyInvitationCompletionResult /*Result*/);
+DECLARE_DELEGATE_FourParams(FOnCancelPartyInvitationComplete, const FUniqueNetId& /*SenderUserId*/, const FOnlinePartyId& /*PartyId*/, const FUniqueNetId& /*TargetUserId*/, const FOnlineError& /*Result*/);
 /**
  * Accepting an invite to a user to join party async task completed callback
  *
@@ -2132,14 +2131,6 @@ enum class ESendPartyInvitationCompletionResult : int8
 	Succeeded = 1
 };
 
-enum class ECancelPartyInvitationCompletionResult : int8
-{
-	InviteNotFound = -100,
-	InvalidParty,
-	UnknownInternalFailure = 0,
-	Succeeded = 1
-};
-
 enum class EAcceptPartyInvitationCompletionResult : int8
 {
 	NotLoggedIn = -100,
@@ -2244,8 +2235,6 @@ ONLINESUBSYSTEM_API const TCHAR* ToString(const ECreatePartyCompletionResult Val
 /** @return the stringified version of the enum passed in */
 ONLINESUBSYSTEM_API const TCHAR* ToString(const ESendPartyInvitationCompletionResult Value);
 /** @return the stringified version of the enum passed in */
-ONLINESUBSYSTEM_API const TCHAR* ToString(const ECancelPartyInvitationCompletionResult Value);
-/** @return the stringified version of the enum passed in */
 ONLINESUBSYSTEM_API const TCHAR* ToString(const EJoinPartyCompletionResult Value);
 /** @return the stringified version of the enum passed in */
 ONLINESUBSYSTEM_API const TCHAR* ToString(const ELeavePartyCompletionResult Value);
