Fix build in OnlineStoreInterfaceV2

diff --git a/Engine/Plugins/Online/OnlineSubsystem/Source/Public/Interfaces/OnlineStoreInterfaceV2.h b/Engine/Plugins/Online/OnlineSubsystem/Source/Public/Interfaces/OnlineStoreInterfaceV2.h
--- a/Engine/Plugins/Online/OnlineSubsystem/Source/Public/Interfaces/OnlineStoreInterfaceV2.h
+++ b/Engine/Plugins/Online/OnlineSubsystem/Source/Public/Interfaces/OnlineStoreInterfaceV2.h
@@ -4,6 +4,7 @@
 
 #include "CoreMinimal.h"
 #include "OnlineDelegateMacros.h"
+#include "OnlineStoreInterfaceV2.generated.h"
 
 ONLINESUBSYSTEM_API DECLARE_LOG_CATEGORY_EXTERN(LogOnlineStoreV2, Log, All);
 
