Delete deprecated OnlineSubsystemTwitch

diff --git a/Engine/Config/BaseEngine.ini b/Engine/Config/BaseEngine.ini
--- a/Engine/Config/BaseEngine.ini
+++ b/Engine/Config/BaseEngine.ini
@@ -2039,19 +2039,6 @@ FriendsFields=locale
 +FriendsFields=link
 +FriendsFields=gender
 
-[OnlineSubsystemTwitch]
-bEnabled=false
-ClientId=""
-
-[OnlineSubsystemTwitch.OnlineIdentityTwitch]
-LoginUrl="https://api.twitch.tv/kraken/oauth2/authorize"
-bForceVerify=true
-LoginRedirectUrl=""
-+LoginDomains=".twitch.tv"
-TokenValidateUrl="https://api.twitch.tv/kraken"
-+ScopeFields=user_read
-TokenRevokeUrl="https://api.twitch.tv/kraken/oauth2/revoke"
-
 [OnlineSubsystemSamsung]
 bEnabled=true
 
diff --git a/Engine/Plugins/Online/OnlineSubsystem/Source/Private/OnlineSubsystem.cpp b/Engine/Plugins/Online/OnlineSubsystem/Source/Private/OnlineSubsystem.cpp
--- a/Engine/Plugins/Online/OnlineSubsystem/Source/Private/OnlineSubsystem.cpp
+++ b/Engine/Plugins/Online/OnlineSubsystem/Source/Private/OnlineSubsystem.cpp
@@ -63,7 +63,6 @@ ONLINESUBSYSTEM_API DEFINE_STAT(STAT_Voice_Interface);
 
 PRAGMA_DISABLE_DEPRECATION_WARNINGS
 FName WECHAT_SUBSYSTEM(TEXT("WeChat"));
-FName TWITCH_SUBSYSTEM(TEXT("TWITCH"));
 FName LIVE_SUBSYSTEM(TEXT("LIVE"));
 FName LIVESERVER_SUBSYSTEM(TEXT("LIVESERVER"));
 PRAGMA_ENABLE_DEPRECATION_WARNINGS
diff --git a/Engine/Plugins/Online/OnlineSubsystem/Source/Public/OnlineSubsystemNames.h b/Engine/Plugins/Online/OnlineSubsystem/Source/Public/OnlineSubsystemNames.h
--- a/Engine/Plugins/Online/OnlineSubsystem/Source/Public/OnlineSubsystemNames.h
+++ b/Engine/Plugins/Online/OnlineSubsystem/Source/Public/OnlineSubsystemNames.h
@@ -78,7 +78,6 @@
 #endif
 
 UE_DEPRECATED(4.27, "OnlineSubsystemWeChat has been deprecated and will be removed") ONLINESUBSYSTEM_API extern FName WECHAT_SUBSYSTEM;
-UE_DEPRECATED(4.27, "OnlineSubsystemTwitch has been deprecated and will be removed") ONLINESUBSYSTEM_API extern FName TWITCH_SUBSYSTEM;
 UE_DEPRECATED(4.27, "OnlineSubsystemLive has been deprecated and will be removed") ONLINESUBSYSTEM_API extern FName LIVE_SUBSYSTEM;
 UE_DEPRECATED(4.27, "OnlineSubsystemLiveServer has been deprecated and will be removed") ONLINESUBSYSTEM_API extern FName LIVESERVER_SUBSYSTEM;
 
diff --git a/Engine/Plugins/Online/OnlineSubsystemTwitch/OnlineSubsystemTwitch.uplugin b/Engine/Plugins/Online/OnlineSubsystemTwitch/OnlineSubsystemTwitch.uplugin
deleted file mode 100644
diff --git a/Engine/Plugins/Online/OnlineSubsystemTwitch/Source/OnlineSubsystemTwitch.Build.cs b/Engine/Plugins/Online/OnlineSubsystemTwitch/Source/OnlineSubsystemTwitch.Build.cs
deleted file mode 100644
diff --git a/Engine/Plugins/Online/OnlineSubsystemTwitch/Source/Private/OnlineAccountTwitch.cpp b/Engine/Plugins/Online/OnlineSubsystemTwitch/Source/Private/OnlineAccountTwitch.cpp
deleted file mode 100644
diff --git a/Engine/Plugins/Online/OnlineSubsystemTwitch/Source/Private/OnlineExternalUITwitch.cpp b/Engine/Plugins/Online/OnlineSubsystemTwitch/Source/Private/OnlineExternalUITwitch.cpp
deleted file mode 100644
diff --git a/Engine/Plugins/Online/OnlineSubsystemTwitch/Source/Private/OnlineExternalUITwitch.h b/Engine/Plugins/Online/OnlineSubsystemTwitch/Source/Private/OnlineExternalUITwitch.h
deleted file mode 100644
diff --git a/Engine/Plugins/Online/OnlineSubsystemTwitch/Source/Private/OnlineIdentityTwitch.cpp b/Engine/Plugins/Online/OnlineSubsystemTwitch/Source/Private/OnlineIdentityTwitch.cpp
deleted file mode 100644
diff --git a/Engine/Plugins/Online/OnlineSubsystemTwitch/Source/Private/OnlineSubsystemModuleTwitch.cpp b/Engine/Plugins/Online/OnlineSubsystemTwitch/Source/Private/OnlineSubsystemModuleTwitch.cpp
deleted file mode 100644
diff --git a/Engine/Plugins/Online/OnlineSubsystemTwitch/Source/Private/OnlineSubsystemTwitch.cpp b/Engine/Plugins/Online/OnlineSubsystemTwitch/Source/Private/OnlineSubsystemTwitch.cpp
deleted file mode 100644
diff --git a/Engine/Plugins/Online/OnlineSubsystemTwitch/Source/Private/OnlineSubsystemTwitchPrivate.h b/Engine/Plugins/Online/OnlineSubsystemTwitch/Source/Private/OnlineSubsystemTwitchPrivate.h
deleted file mode 100644
diff --git a/Engine/Plugins/Online/OnlineSubsystemTwitch/Source/Private/TwitchTokenValidationResponse.h b/Engine/Plugins/Online/OnlineSubsystemTwitch/Source/Private/TwitchTokenValidationResponse.h
deleted file mode 100644
diff --git a/Engine/Plugins/Online/OnlineSubsystemTwitch/Source/Public/OnlineAccountTwitch.h b/Engine/Plugins/Online/OnlineSubsystemTwitch/Source/Public/OnlineAccountTwitch.h
deleted file mode 100644
diff --git a/Engine/Plugins/Online/OnlineSubsystemTwitch/Source/Public/OnlineIdentityTwitch.h b/Engine/Plugins/Online/OnlineSubsystemTwitch/Source/Public/OnlineIdentityTwitch.h
deleted file mode 100644
diff --git a/Engine/Plugins/Online/OnlineSubsystemTwitch/Source/Public/OnlineSubsystemTwitch.h b/Engine/Plugins/Online/OnlineSubsystemTwitch/Source/Public/OnlineSubsystemTwitch.h
deleted file mode 100644
diff --git a/Engine/Plugins/Online/OnlineSubsystemTwitch/Source/Public/OnlineSubsystemTwitchModule.h b/Engine/Plugins/Online/OnlineSubsystemTwitch/Source/Public/OnlineSubsystemTwitchModule.h
deleted file mode 100644
diff --git a/Engine/Plugins/Online/OnlineSubsystemTwitch/Source/Public/OnlineSubsystemTwitchPackage.h b/Engine/Plugins/Online/OnlineSubsystemTwitch/Source/Public/OnlineSubsystemTwitchPackage.h
deleted file mode 100644
diff --git a/Engine/Plugins/Online/OnlineSubsystemUtils/Source/OnlineSubsystemUtils/Private/OnlineSubsystemUtilsModule.cpp b/Engine/Plugins/Online/OnlineSubsystemUtils/Source/OnlineSubsystemUtils/Private/OnlineSubsystemUtilsModule.cpp
--- a/Engine/Plugins/Online/OnlineSubsystemUtils/Source/OnlineSubsystemUtils/Private/OnlineSubsystemUtilsModule.cpp
+++ b/Engine/Plugins/Online/OnlineSubsystemUtils/Source/OnlineSubsystemUtils/Private/OnlineSubsystemUtilsModule.cpp
@@ -212,7 +212,6 @@ PRAGMA_ENABLE_DEPRECATION_WARNINGS
 		CREATE_HASH(THUNDERHEAD_SUBSYSTEM); 
 PRAGMA_DISABLE_DEPRECATION_WARNINGS
 		CREATE_HASH(WECHAT_SUBSYSTEM);
-		CREATE_HASH(TWITCH_SUBSYSTEM); 
 PRAGMA_ENABLE_DEPRECATION_WARNINGS
 		CREATE_HASH(OCULUS_SUBSYSTEM); 
 		CREATE_HASH(QUAIL_SUBSYSTEM);
