Addition of OnlineBase.h for common definitions between OnlineServices and OnlineSubsystem

diff --git a/Engine/Plugins/Online/OnlineBase/Source/Public/Online/LANBeacon.h b/Engine/Plugins/Online/OnlineBase/Source/Public/Online/LANBeacon.h
--- a/Engine/Plugins/Online/OnlineBase/Source/Public/Online/LANBeacon.h
+++ b/Engine/Plugins/Online/OnlineBase/Source/Public/Online/LANBeacon.h
@@ -5,18 +5,7 @@
 #include "CoreMinimal.h"
 #include "Misc/ConfigCacheIni.h"
 #include "Misc/Guid.h"
-
-#ifndef ONLINE_SUCCESS
-#define ONLINE_SUCCESS 0
-#endif
-
-#ifndef ONLINE_FAIL
-#define ONLINE_FAIL (uint32)-1
-#endif
-
-#ifndef ONLINE_IO_PENDING
-#define ONLINE_IO_PENDING 997
-#endif
+#include "Online/OnlineBase.h"
 
 /**
  * Generates a random nonce (number used once) of the desired length
diff --git a/Engine/Plugins/Online/OnlineBase/Source/Public/Online/OnlineBase.h b/Engine/Plugins/Online/OnlineBase/Source/Public/Online/OnlineBase.h
new file mode 100644
--- /dev/null
+++ b/Engine/Plugins/Online/OnlineBase/Source/Public/Online/OnlineBase.h
@@ -0,0 +1,15 @@
+// Copyright Epic Games, Inc. All Rights Reserved.
+
+#pragma once
+
+#ifndef ONLINE_SUCCESS
+#define ONLINE_SUCCESS 0
+#endif
+
+#ifndef ONLINE_FAIL
+#define ONLINE_FAIL (uint32)-1
+#endif
+
+#ifndef ONLINE_IO_PENDING
+#define ONLINE_IO_PENDING 997
+#endif
\ No newline at end of file
diff --git a/Engine/Plugins/Online/OnlineSubsystem/Source/Public/OnlineSubsystemTypes.h b/Engine/Plugins/Online/OnlineSubsystem/Source/Public/OnlineSubsystemTypes.h
--- a/Engine/Plugins/Online/OnlineSubsystem/Source/Public/OnlineSubsystemTypes.h
+++ b/Engine/Plugins/Online/OnlineSubsystem/Source/Public/OnlineSubsystemTypes.h
@@ -4,6 +4,7 @@
 
 #include "CoreMinimal.h"
 #include "Online/CoreOnline.h"
+#include "Online/OnlineBase.h"
 #include "OnlineSubsystemNames.h"  // can be removed once we have no more temporary FUniqueNetId subtypes
 #include "OnlineSubsystemPackage.h"
 
@@ -25,18 +26,6 @@ extern ONLINESUBSYSTEM_API bool IsUniqueIdLocal(const FUniqueNetId& UniqueId);
 
 #define DEDICATED_SERVER_USER_INDEX 0
 
-#ifndef ONLINE_SUCCESS
-#define ONLINE_SUCCESS 0
-#endif
-
-#ifndef ONLINE_FAIL
-#define ONLINE_FAIL (uint32)-1
-#endif
-
-#ifndef ONLINE_IO_PENDING
-#define ONLINE_IO_PENDING 997
-#endif
-
 /**
  * Environment for the current online platform
  */
