[OnlineFramework]
* Added so McpRegion is showing as console status field
* Added so BuildId is showing as console status field
* Added so Backend is showing as console status field

diff --git a/Engine/Plugins/Online/OnlineFramework/Source/Qos/Private/QosRegionManager.cpp b/Engine/Plugins/Online/OnlineFramework/Source/Qos/Private/QosRegionManager.cpp
--- a/Engine/Plugins/Online/OnlineFramework/Source/Qos/Private/QosRegionManager.cpp
+++ b/Engine/Plugins/Online/OnlineFramework/Source/Qos/Private/QosRegionManager.cpp
@@ -3,6 +3,7 @@
 #include "QosRegionManager.h"
 #include "Misc/CommandLine.h"
 #include "Misc/ConfigCacheIni.h"
+#include "Misc/TrackedActivity.h"
 #include "QosInterface.h"
 #include "TimerManager.h"
 #include "Engine/World.h"
@@ -186,6 +187,9 @@ UQosRegionManager::UQosRegionManager(const FObjectInitializer& ObjectInitializer
 	{
 		ForceRegionId.ToUpperInline();
 	}
+
+	// Will add an info entry to the console
+	static FTrackedActivity Ta(TEXT("McpRegion"), *ForceRegionId, FTrackedActivity::ELight::None, FTrackedActivity::EType::Info);
 }
 
 void UQosRegionManager::PostReloadConfig(FProperty* PropertyThatWasLoaded)
diff --git a/Engine/Plugins/Online/OnlineSubsystem/Source/Private/OnlineSubsystem.cpp b/Engine/Plugins/Online/OnlineSubsystem/Source/Private/OnlineSubsystem.cpp
--- a/Engine/Plugins/Online/OnlineSubsystem/Source/Private/OnlineSubsystem.cpp
+++ b/Engine/Plugins/Online/OnlineSubsystem/Source/Private/OnlineSubsystem.cpp
@@ -3,6 +3,7 @@
 #include "OnlineSubsystem.h"
 #include "Misc/CommandLine.h"
 #include "Misc/ConfigCacheIni.h"
+#include "Misc/TrackedActivity.h"
 #include "HAL/IConsoleManager.h"
 #include "Online/NboSerializer.h"
 #include "Misc/NetworkVersion.h"
@@ -178,6 +179,10 @@ int32 GetBuildUniqueId()
 		{
 			BuildId = BuildIdOverride;
 		}
+
+		// Will add an info entry to the console
+		static FTrackedActivity Ta(TEXT("BuildId"), *FString::FromInt(BuildId), FTrackedActivity::ELight::None, FTrackedActivity::EType::Info);
+		CVarBuildIdOverride->SetOnChangedCallback(FConsoleVariableDelegate::CreateLambda([](IConsoleVariable* CVar) { Ta.Update(*CVar->GetString()); }));
 	}
 
 	return BuildId;
diff --git a/Engine/Plugins/Online/OnlineSubsystem/Source/Private/OnlineSubsystemImpl.cpp b/Engine/Plugins/Online/OnlineSubsystem/Source/Private/OnlineSubsystemImpl.cpp
--- a/Engine/Plugins/Online/OnlineSubsystem/Source/Private/OnlineSubsystemImpl.cpp
+++ b/Engine/Plugins/Online/OnlineSubsystem/Source/Private/OnlineSubsystemImpl.cpp
@@ -33,6 +33,8 @@ const FName FOnlineSubsystemImpl::DefaultInstanceName(TEXT("DefaultInstance"));
 FOnlineSubsystemImpl::FOnlineSubsystemImpl(FName InSubsystemName, FName InInstanceName) :
 	FOnlineSubsystemImpl(InSubsystemName, InInstanceName, FTSTicker::GetCoreTicker())
 {
+	// This is just here to prime BuildId variable and add as info in console
+	GetBuildUniqueId();
 }
 
 FOnlineSubsystemImpl::FOnlineSubsystemImpl(FName InSubsystemName, FName InInstanceName, FTSTicker& Ticker) :
