- Add tracking to OnlineAsyncTaskManager to catch hung serial tasks / online thread Tick()

diff --git a/Engine/Plugins/Online/OnlineSubsystem/Source/Private/OnlineAsyncTaskManager.cpp b/Engine/Plugins/Online/OnlineSubsystem/Source/Private/OnlineAsyncTaskManager.cpp
--- a/Engine/Plugins/Online/OnlineSubsystem/Source/Private/OnlineAsyncTaskManager.cpp
+++ b/Engine/Plugins/Online/OnlineSubsystem/Source/Private/OnlineAsyncTaskManager.cpp
@@ -28,7 +28,7 @@ namespace OSSConsoleVariables
 #define POLLING_INTERVAL_MS 20
 
 FOnlineAsyncTaskManager::FOnlineAsyncTaskManager() :
-	ActiveTask(nullptr),
+	ActiveSerialTask(nullptr),
 	MaxParallelTasks(8),
 	bReloadMaxParallelTasksConfig(false),
 	WorkEvent(FPlatformProcess::GetSynchEventFromPool()),
@@ -43,6 +43,8 @@ FOnlineAsyncTaskManager::FOnlineAsyncTaskManager() :
 	}
 
 	GConfig->GetInt(TEXT("OnlineSubsystem"), TEXT("MaxParallelTasks"), MaxParallelTasks, GEngineIni);
+	GConfig->GetBool(TEXT("OnlineSubsystem"), TEXT("bEnableReportBreach"), bEnableReportBreach, GEngineIni);
+	GConfig->GetFloat(TEXT("OnlineSubsystem"), TEXT("BreachTimeSeconds"), ConfigBreachTimeSeconds, GEngineIni);
 }
 
 FOnlineAsyncTaskManager::~FOnlineAsyncTaskManager()
@@ -82,8 +84,8 @@ void FOnlineAsyncTaskManager::Stop(void)
 {
 	int32 NumInTasks = 0;
 	{
-		FScopeLock Lock(&InQueueLock);
-		NumInTasks = InQueue.Num();
+		FScopeLock Lock(&QueuedSerialTasksLock);
+		NumInTasks = QueuedSerialTasks.Num();
 	}
 
 	int32 NumOutTasks = 0;
@@ -92,7 +94,7 @@ void FOnlineAsyncTaskManager::Stop(void)
 		NumOutTasks = OutQueue.Num();
 	}
 
-	UE_LOG_ONLINE(VeryVerbose, TEXT("FOnlineAsyncTaskManager::Stop() ActiveTask:%p Tasks[%d/%d]"), ActiveTask, NumInTasks, NumOutTasks);
+	UE_LOG_ONLINE(VeryVerbose, TEXT("FOnlineAsyncTaskManager::Stop() ActiveSerialTask:%p Tasks[%d/%d]"), ActiveSerialTask, NumInTasks, NumOutTasks);
 
 	// Set the variable to requesting exit before we trigger the event
 	bRequestingExit = true;
@@ -111,18 +113,24 @@ void FOnlineAsyncTaskManager::Exit(void)
 
 void FOnlineAsyncTaskManager::AddToInQueue(FOnlineAsyncTask* NewTask)
 {
-	FScopeLock Lock(&InQueueLock);
-	InQueue.Add(NewTask);
+	UE_LOG_ONLINE(VeryVerbose, TEXT("FOnlineAsyncTaskManager::AddToInQueue [%s]"), *NewTask->ToString());
+
+	FScopeLock Lock(&QueuedSerialTasksLock);
+	QueuedSerialTasks.Add(NewTask);
 }
 
 void FOnlineAsyncTaskManager::AddToOutQueue(FOnlineAsyncItem* CompletedItem)
 {
+	UE_LOG_ONLINE(VeryVerbose, TEXT("FOnlineAsyncTaskManager::AddToOutQueue [%s]"), *CompletedItem->ToString());
+
 	FScopeLock Lock(&OutQueueLock);
 	OutQueue.Add(CompletedItem);
 }
 
 void FOnlineAsyncTaskManager::AddToParallelTasks(FOnlineAsyncTask* NewTask)
 {
+	UE_LOG_ONLINE(VeryVerbose, TEXT("FOnlineAsyncTaskManager::AddToParallelTasks [%s]"), *NewTask->ToString());
+
 	bReloadMaxParallelTasksConfig = true;
 
 	QueuedParallelTasks.Enqueue(NewTask);
@@ -130,6 +138,8 @@ void FOnlineAsyncTaskManager::AddToParallelTasks(FOnlineAsyncTask* NewTask)
 
 void FOnlineAsyncTaskManager::RemoveFromParallelTasks(FOnlineAsyncTask* OldTask)
 {
+	UE_LOG_ONLINE(VeryVerbose, TEXT("FOnlineAsyncTaskManager::RemoveFromParallelTasks [%s]"), *OldTask->ToString());
+
 	FScopeLock LockParallelTasks(&ParallelTasksLock);
 
 	ParallelTasks.Remove( OldTask );
@@ -155,6 +165,8 @@ void FOnlineAsyncTaskManager::GameTick()
 	FOnlineAsyncTask* ParallelTask = nullptr;
 	while (NumParallelTasksToStart-- > 0 && QueuedParallelTasks.Dequeue(ParallelTask))
 	{
+		UE_LOG_ONLINE(VeryVerbose, TEXT("FOnlineAsyncTaskManager::GameTick Starting parallel task [%s]"), *ParallelTask->ToString());
+
 		ParallelTask->Initialize();
 
 		FScopeLock LockParallelTasks(&ParallelTasksLock);
@@ -215,42 +227,86 @@ void FOnlineAsyncTaskManager::GameTick()
 	}
 	while (CurrentQueueSize > 1);
 
-	int32 QueueSize = 0;
-	bool bHasActiveTask = false;
+	const double TimeNowSeconds = FPlatformTime::Seconds();
+
+	// Detect hung online thread tick
+	if(bEnableReportBreach)
 	{
+		FOnlineAsyncTask* HungTask = nullptr;
+		bool bReportHungTick = false;
 		{
-			FScopeLock LockInQueue(&InQueueLock);
-			QueueSize = InQueue.Num();
-		}
-		{
-			FScopeLock LockActiveTask(&ActiveTaskLock);
-			if (ActiveTask != nullptr)
+			FScopeLock Lock(&OnlineThreadTickInfoLock);
+			const double BreachTimeSeconds = OnlineThreadTickInfo.TickStartTime + ConfigBreachTimeSeconds;
+			bReportHungTick = OnlineThreadTickInfo.bIsTicking && !OnlineThreadTickInfo.bBreachReported && TimeNowSeconds > BreachTimeSeconds;
+			if (bReportHungTick)
 			{
-				++QueueSize;
-				bHasActiveTask = true;
+				OnlineThreadTickInfo.bBreachReported = true;
+				HungTask = OnlineThreadTickInfo.CurrentTask;
 			}
 		}
+		if (bReportHungTick)
+		{
+			UE_LOG_ONLINE(Warning, TEXT("OnlineAsyncTaskManager::GameTick online thread tick breached, Task=[%s]"), HungTask ? *HungTask->ToString() : TEXT("nullptr"));
+		}
+	}
 
-		if (!bHasActiveTask && QueueSize > 0)
+	int32 QueueSize = 0;
+	{
+		FScopeLock LockInQueue(&QueuedSerialTasksLock);
+		QueueSize = QueuedSerialTasks.Num();
+	}
+	FOnlineAsyncTask* ActiveTask = nullptr;
+	{
+		FScopeLock Lock(&ActiveSerialTaskLock);
+		ActiveTask = ActiveSerialTask;
+	}
+	if (ActiveTask)
+	{
+		++QueueSize;
+		if (bEnableReportBreach)
 		{
-			// Grab the current task from the queue
-			FOnlineAsyncTask* Task = nullptr;
+			const double BreachTimeSeconds = ActiveSerialTaskStartTime + ConfigBreachTimeSeconds;
+			const bool bReportBreach = !bActiveSerialTaskBreachReported && TimeNowSeconds > BreachTimeSeconds;
+			if (bReportBreach)
 			{
-				FScopeLock LockInQueue(&InQueueLock);
-				Task = InQueue[0];
-				InQueue.RemoveAt(0);
-			}
+				bActiveSerialTaskBreachReported = true;
 
-			// Initialize the task before giving it away to the online thread
-			Task->Initialize();
-			{
-				FScopeLock LockActiveTask(&ActiveTaskLock);
-				ActiveTask = Task;
+				UE_LOG_ONLINE(Warning, TEXT("OnlineAsyncTaskManager::GameTick serial task breached, Task=[%s]"), *ActiveTask->ToString());
+				TArray<FOnlineAsyncTask*> SerialTaskQueue;
+				{
+					FScopeLock LockInQueue(&QueuedSerialTasksLock);
+					SerialTaskQueue = QueuedSerialTasks;
+				}
+				for (const FOnlineAsyncTask* Task : SerialTaskQueue)
+				{
+					UE_LOG_ONLINE(Warning, TEXT("	blocked task [%s]"), *Task->ToString());
+				}
 			}
+		}
+	}
+	else if (QueueSize > 0)
+	{
+		// Grab the current task from the queue
+		FOnlineAsyncTask* Task = nullptr;
+		{
+			FScopeLock LockInQueue(&QueuedSerialTasksLock);
+			Task = QueuedSerialTasks[0];
+			QueuedSerialTasks.RemoveAt(0);
+		}
+
+		UE_LOG_ONLINE(VeryVerbose, TEXT("FOnlineAsyncTaskManager::GameTick Starting serial task [%s]"), *Task->ToString());
 
-			// Wake up the online thread
-			WorkEvent->Trigger();
+		// Initialize the task before giving it away to the online thread
+		Task->Initialize();
+		{
+			FScopeLock LockActiveTask(&ActiveSerialTaskLock);
+			ActiveSerialTask = Task;
 		}
+		ActiveSerialTaskStartTime = TimeNowSeconds;
+		bActiveSerialTaskBreachReported = false;
+
+		// Wake up the online thread
+		WorkEvent->Trigger();
 	}
 
 	SET_DWORD_STAT(STAT_Online_AsyncTasks, QueueSize);
@@ -259,6 +315,14 @@ void FOnlineAsyncTaskManager::GameTick()
 void FOnlineAsyncTaskManager::Tick()
 {
 	SCOPE_CYCLE_COUNTER(STAT_Online_Async);
+	if (bEnableReportBreach)
+	{
+		FScopeLock Lock(&OnlineThreadTickInfoLock);
+		OnlineThreadTickInfo.bIsTicking = true;
+		OnlineThreadTickInfo.TickStartTime = FPlatformTime::Seconds();
+		OnlineThreadTickInfo.CurrentTask = nullptr;
+		OnlineThreadTickInfo.bBreachReported = false;
+	}
 
 	// Tick Online services (possibly callbacks). 
 	OnlineTick();
@@ -277,6 +341,13 @@ void FOnlineAsyncTaskManager::Tick()
 		for (auto It = CopyParallelTasks.CreateIterator(); It; ++It)
 		{
 			Task = *It;
+
+			if (bEnableReportBreach)
+			{
+				FScopeLock Lock(&OnlineThreadTickInfoLock);
+				OnlineThreadTickInfo.CurrentTask = Task;
+			}
+
 			Task->Tick();
 
 			if (Task->IsDone())
@@ -299,18 +370,28 @@ void FOnlineAsyncTaskManager::Tick()
 				AddToOutQueue(Task);
 			}
 		}
+		if (bEnableReportBreach)
+		{
+			FScopeLock Lock(&OnlineThreadTickInfoLock);
+			OnlineThreadTickInfo.CurrentTask = nullptr;
+		}
 	}
 
 	{
 		// Now process the serial "In" queue
 		FOnlineAsyncTask* Task = nullptr;
 		{
-			FScopeLock LockActiveTask(&ActiveTaskLock);
-			Task = ActiveTask;
+			FScopeLock Lock(&ActiveSerialTaskLock);
+			Task = ActiveSerialTask;
 		}
 
 		if (Task)
 		{
+			if (bEnableReportBreach)
+			{
+				FScopeLock Lock(&OnlineThreadTickInfoLock);
+				OnlineThreadTickInfo.CurrentTask = Task;
+			}
 			Task->Tick();
 
 			if (Task->IsDone())
@@ -332,10 +413,15 @@ void FOnlineAsyncTaskManager::Tick()
 				AddToOutQueue(Task);
 
 				{
-					FScopeLock LockActiveTask(&ActiveTaskLock);
-					ActiveTask = nullptr;
+					FScopeLock Lock(&ActiveSerialTaskLock);
+					ActiveSerialTask = nullptr;
 				}
 			}
 		}
 	}
+	if (bEnableReportBreach)
+	{
+		FScopeLock Lock(&OnlineThreadTickInfoLock);
+		OnlineThreadTickInfo.bIsTicking = false;
+	}
 }
diff --git a/Engine/Plugins/Online/OnlineSubsystem/Source/Public/OnlineAsyncTaskManager.h b/Engine/Plugins/Online/OnlineSubsystem/Source/Public/OnlineAsyncTaskManager.h
--- a/Engine/Plugins/Online/OnlineSubsystem/Source/Public/OnlineAsyncTaskManager.h
+++ b/Engine/Plugins/Online/OnlineSubsystem/Source/Public/OnlineAsyncTaskManager.h
@@ -304,16 +304,41 @@ class ONLINESUBSYSTEM_API FOnlineAsyncTaskManager : public FRunnable, FSingleThr
 private:
 	
 	/** The current active task processed serially by the async task manager */
-	FOnlineAsyncTask* ActiveTask;
-	/** Critical section for modifying the active task */
-	FCriticalSection ActiveTaskLock;
+	FOnlineAsyncTask* ActiveSerialTask = nullptr;
+	/** Critical section for modifying the active serial task */
+	FCriticalSection ActiveSerialTaskLock;
+	/** When the active serial task was started on the game thread */
+	double ActiveSerialTaskStartTime = 0.0;
+	/** Has a breach been reported for the active task */
+	bool bActiveSerialTaskBreachReported = false;
+
+	struct FOnlineThreadTickInfo
+	{
+		/* Is the online thread currently ticking */
+		bool bIsTicking = false;
+		/* Time the online thread began this tick */
+		double TickStartTime = 0.0;
+		/* Task the online thread is currently ticking */
+		FOnlineAsyncTask* CurrentTask = nullptr;
+		/* */
+		bool bBreachReported = false;
+	};
+	/** Info on what the online thread is currently doing */
+	FOnlineThreadTickInfo OnlineThreadTickInfo;
+	/** Critical section for modifying this tick info */
+	FCriticalSection OnlineThreadTickInfoLock;
+
+	/** Should we report a breach or not */
+	bool bEnableReportBreach = false;
+	/** Time to allow a serial task / the online thread tick to run before reporting */
+	float ConfigBreachTimeSeconds = 60.0;
 
 protected:
 
 	/** Game thread async tasks are queued up here for processing on the online thread */
-	TArray<FOnlineAsyncTask*> InQueue;
+	TArray<FOnlineAsyncTask*> QueuedSerialTasks;
 	/** Critical section for thread safe operation of the event in queue */
-	FCriticalSection InQueueLock;
+	FCriticalSection QueuedSerialTasksLock;
 
 	/** Number of tasks that can run in parallel */
 	int32 MaxParallelTasks;
