Add new message array separator character

diff --git a/Engine/Plugins/Online/OnlineSubsystem/Source/Public/Interfaces/IMessageSanitizerInterface.h b/Engine/Plugins/Online/OnlineSubsystem/Source/Public/Interfaces/IMessageSanitizerInterface.h
--- a/Engine/Plugins/Online/OnlineSubsystem/Source/Public/Interfaces/IMessageSanitizerInterface.h
+++ b/Engine/Plugins/Online/OnlineSubsystem/Source/Public/Interfaces/IMessageSanitizerInterface.h
@@ -46,6 +46,8 @@ public:
 
 	/** Invalidate all previously queried blocked users state */
 	virtual void ResetBlockedUserCache() = 0;
+
+	const FString MessageArraySeparator = TEXT("\u241e");
 };
 
 typedef TSharedPtr<IMessageSanitizer, ESPMode::ThreadSafe> IMessageSanitizerPtr;
