Temporary fix to PVS-Studio issue with inlined generated CPP in this module.

diff --git a/Engine/Plugins/Online/OnlineSubsystem/Source/Private/NamedInterfaces.cpp b/Engine/Plugins/Online/OnlineSubsystem/Source/Private/NamedInterfaces.cpp
--- a/Engine/Plugins/Online/OnlineSubsystem/Source/Private/NamedInterfaces.cpp
+++ b/Engine/Plugins/Online/OnlineSubsystem/Source/Private/NamedInterfaces.cpp
@@ -4,7 +4,8 @@
 #include "UObject/Package.h"
 #include "OnlineSubsystem.h"
 
-#include UE_INLINE_GENERATED_CPP_BY_NAME(NamedInterfaces)
+// Temporary disabling of the inline of generated cpp to avoid PVS-Studio issue until we can get it resolved
+//#include UE_DISABLE_INLINE_GENERATED_CPP_BY_NAME(NamedInterfaces)
 
 UNamedInterfaces::UNamedInterfaces(const FObjectInitializer& ObjectInitializer) :
 	Super(ObjectInitializer)
diff --git a/Engine/Plugins/Online/OnlineSubsystem/Source/Private/TurnBasedMatchInterface.cpp b/Engine/Plugins/Online/OnlineSubsystem/Source/Private/TurnBasedMatchInterface.cpp
--- a/Engine/Plugins/Online/OnlineSubsystem/Source/Private/TurnBasedMatchInterface.cpp
+++ b/Engine/Plugins/Online/OnlineSubsystem/Source/Private/TurnBasedMatchInterface.cpp
@@ -2,7 +2,8 @@
 
 #include "Interfaces/TurnBasedMatchInterface.h"
 
-#include UE_INLINE_GENERATED_CPP_BY_NAME(TurnBasedMatchInterface)
+// Temporary disabling of the inline of generated cpp to avoid PVS-Studio issue until we can get it resolved
+//#include UE_DISABLE_INLINE_GENERATED_CPP_BY_NAME(TurnBasedMatchInterface)
 
 UTurnBasedMatchInterface::UTurnBasedMatchInterface(const FObjectInitializer& ObjectInitializer)
 : Super(ObjectInitializer)
