PurcahseCatalogEntry now fires on the server if triggered mid-match. It will then forward the notification to the client.

diff --git a/Engine/Plugins/Online/OnlineSubsystem/Source/Private/OnlineNotification.cpp b/Engine/Plugins/Online/OnlineSubsystem/Source/Private/OnlineNotification.cpp
--- a/Engine/Plugins/Online/OnlineSubsystem/Source/Private/OnlineNotification.cpp
+++ b/Engine/Plugins/Online/OnlineSubsystem/Source/Private/OnlineNotification.cpp
@@ -9,12 +9,14 @@ FOnlineNotification::FOnlineNotification(
 	const FString& InTypeStr,
 	const TSharedPtr<FJsonValue>& InPayload,
 	FUniqueNetIdPtr InToUserId,
-	FUniqueNetIdPtr InFromUserId
+	FUniqueNetIdPtr InFromUserId,
+	const FString& InClientRequestIdStr
 )
 : TypeStr(InTypeStr)
 , Payload(InPayload.IsValid() ? InPayload->AsObject() : nullptr)
 , ToUserId(InToUserId)
 , FromUserId(InFromUserId)
+, ClientRequestIdStr(InClientRequestIdStr)
 {
 }
 
@@ -36,3 +38,15 @@ void FOnlineNotification::SetTypeFromPayload()
 		}
 	}
 }
+
+void FOnlineNotification::SetClientRequestIdFromPayload()
+{
+	if (Payload.IsValid() && ClientRequestIdStr.IsEmpty())
+	{
+		if (!Payload->TryGetStringField(TEXT("client_request_id"), ClientRequestIdStr))
+		{
+			UE_LOG_ONLINE(Error, TEXT("No client_request_id in notification JSON object"));
+			ClientRequestIdStr = TEXT("");
+		}
+	}
+}
diff --git a/Engine/Plugins/Online/OnlineSubsystem/Source/Public/OnlineNotification.h b/Engine/Plugins/Online/OnlineSubsystem/Source/Public/OnlineNotification.h
--- a/Engine/Plugins/Online/OnlineSubsystem/Source/Public/OnlineNotification.h
+++ b/Engine/Plugins/Online/OnlineSubsystem/Source/Public/OnlineNotification.h
@@ -26,7 +26,8 @@ struct ONLINESUBSYSTEM_API FOnlineNotification
 		const FString& InTypeStr,
 		const TSharedPtr<FJsonValue>& InPayload,
 		FUniqueNetIdPtr InToUserId = nullptr,
-		FUniqueNetIdPtr InFromUserId = nullptr
+		FUniqueNetIdPtr InFromUserId = nullptr,
+		const FString & InClientRequestIdStr = TEXT("")
 	);
 
 	/**
@@ -37,12 +38,14 @@ struct ONLINESUBSYSTEM_API FOnlineNotification
 		const FString& InTypeStr,
 		const TSharedPtr<FJsonObject>& InPayload,
 		FUniqueNetIdPtr InToUserId = nullptr,
-		FUniqueNetIdPtr InFromUserId = nullptr
+		FUniqueNetIdPtr InFromUserId = nullptr,
+		const FString& InClientRequestIdStr = TEXT("")
 	)
 	: TypeStr(InTypeStr)
 	, Payload(InPayload)
 	, ToUserId(InToUserId)
 	, FromUserId(InFromUserId)
+	, ClientRequestIdStr(InClientRequestIdStr)
 	{
 	}
 
@@ -74,6 +77,11 @@ struct ONLINESUBSYSTEM_API FOnlineNotification
 	 */
 	void SetTypeFromPayload();
 
+	/**
+	 * Set up the ClientRequestIdStr for the case where it is embedded in the payload
+	 */
+	void SetClientRequestIdFromPayload();
+
 	/** A string defining the type of this notification, used to determine how to parse the payload */
 	FString TypeStr;
 
@@ -85,4 +93,7 @@ struct ONLINESUBSYSTEM_API FOnlineNotification
 
 	/** User who sent the notification, optional. */
 	FUniqueNetIdPtr FromUserId;
+
+	/** String representing the client_request_id for this notification. Used to tie a server request back to the client. Can be empty. */
+	FString ClientRequestIdStr;
 };
