Increasing PVS-Studio timeout.

#rnx

diff --git a/Engine/Plugins/Online/OnlineSubsystem/Source/Private/NamedInterfaces.cpp b/Engine/Plugins/Online/OnlineSubsystem/Source/Private/NamedInterfaces.cpp
--- a/Engine/Plugins/Online/OnlineSubsystem/Source/Private/NamedInterfaces.cpp
+++ b/Engine/Plugins/Online/OnlineSubsystem/Source/Private/NamedInterfaces.cpp
@@ -4,8 +4,7 @@
 #include "UObject/Package.h"
 #include "OnlineSubsystem.h"
 
-// Temporary disabling of the inline of generated cpp to avoid PVS-Studio issue until we can get it resolved
-//#include UE_DISABLE_INLINE_GENERATED_CPP_BY_NAME(NamedInterfaces)
+#include UE_INLINE_GENERATED_CPP_BY_NAME(NamedInterfaces)
 
 UNamedInterfaces::UNamedInterfaces(const FObjectInitializer& ObjectInitializer) :
 	Super(ObjectInitializer)
diff --git a/Engine/Plugins/Online/OnlineSubsystem/Source/Private/TurnBasedMatchInterface.cpp b/Engine/Plugins/Online/OnlineSubsystem/Source/Private/TurnBasedMatchInterface.cpp
--- a/Engine/Plugins/Online/OnlineSubsystem/Source/Private/TurnBasedMatchInterface.cpp
+++ b/Engine/Plugins/Online/OnlineSubsystem/Source/Private/TurnBasedMatchInterface.cpp
@@ -2,8 +2,7 @@
 
 #include "Interfaces/TurnBasedMatchInterface.h"
 
-// Temporary disabling of the inline of generated cpp to avoid PVS-Studio issue until we can get it resolved
-//#include UE_DISABLE_INLINE_GENERATED_CPP_BY_NAME(TurnBasedMatchInterface)
+#include UE_INLINE_GENERATED_CPP_BY_NAME(TurnBasedMatchInterface)
 
 UTurnBasedMatchInterface::UTurnBasedMatchInterface(const FObjectInitializer& ObjectInitializer)
 : Super(ObjectInitializer)
diff --git a/Engine/Source/Programs/UnrealBuildTool/Platform/Windows/PVSToolChain.cs b/Engine/Source/Programs/UnrealBuildTool/Platform/Windows/PVSToolChain.cs
--- a/Engine/Source/Programs/UnrealBuildTool/Platform/Windows/PVSToolChain.cs
+++ b/Engine/Source/Programs/UnrealBuildTool/Platform/Windows/PVSToolChain.cs
@@ -251,7 +251,7 @@ namespace UnrealBuildTool
 				}
 				else
 				{
-					return AnalysisTimeoutFlags.After_10_minutes;
+					return AnalysisTimeoutFlags.After_30_minutes;
 				}
 			}
 			set
