Rename file OnlineSessionInterfaceGDK.h/.cpp to OnlineSesssionInterfaceMpsdGDK.h/.cpp;
Rename OnlineSessionGDK to OnlineSessionMpsdGDK;
Add OnlineSessionMpaGDK;
Add OnlineSessionGDK.h with OnlineSessionGDK class.
Move MAX_RETRIES and SessionUpdateStatName from OnlineSessionMpsdGDK to OnlineSessionGDK because they will be used by MPA impl as well.

diff --git a/Engine/Plugins/Online/OnlineSubsystem/Source/Public/OnlineSessionSettings.h b/Engine/Plugins/Online/OnlineSubsystem/Source/Public/OnlineSessionSettings.h
--- a/Engine/Plugins/Online/OnlineSubsystem/Source/Public/OnlineSessionSettings.h
+++ b/Engine/Plugins/Online/OnlineSubsystem/Source/Public/OnlineSessionSettings.h
@@ -526,6 +526,9 @@ public:
 	}
 };
 
+using FNamedOnlineSessionPtr = TSharedPtr<FNamedOnlineSession, ESPMode::ThreadSafe>;
+using FNamedOnlineSessionRef = TSharedRef<FNamedOnlineSession, ESPMode::ThreadSafe>;
+
 /** Value returned on unreachable or otherwise bad search results */
 #define MAX_QUERY_PING 9999
 
