Change feature flags so that a runtime check cannot enable a feature which was defined as off.

diff --git a/Engine/Plugins/Online/OnlineSubsystem/Source/Public/OnlineSubsystemFeatureFlags.h b/Engine/Plugins/Online/OnlineSubsystem/Source/Public/OnlineSubsystemFeatureFlags.h
--- a/Engine/Plugins/Online/OnlineSubsystem/Source/Public/OnlineSubsystemFeatureFlags.h
+++ b/Engine/Plugins/Online/OnlineSubsystem/Source/Public/OnlineSubsystemFeatureFlags.h
@@ -46,11 +46,15 @@ namespace UE::OnlineSubsystemFeature { \
 	FOnlineSubsystemFeatureCheck FeatureRuntimeCheckDelegate_##FeatureName; \
 	bool RuntimeCheckFeature_##FeatureName() \
 	{ \
-		static const bool Runtime##FeatureName##Enabled = \
-			FeatureRuntimeCheckDelegate_##FeatureName.IsBound() ? \
-			FeatureRuntimeCheckDelegate_##FeatureName.Execute() : \
-			ONLINE_SUBSYSTEM_FEATURE_##FeatureName; \
-		return Runtime##FeatureName##Enabled; \
+		if constexpr (ONLINE_SUBSYSTEM_FEATURE_##FeatureName) \
+		{ \
+			static const bool Runtime##FeatureName##Enabled = FeatureRuntimeCheckDelegate_##FeatureName.IsBound() ? FeatureRuntimeCheckDelegate_##FeatureName.Execute() : true; \
+			return Runtime##FeatureName##Enabled; \
+		} \
+		else \
+		{ \
+			return false; \
+		} \
 	} \
 }
 
