Implement AddRecentPlayers in friends gdk interface.

diff --git a/Engine/Plugins/Online/OnlineSubsystem/Source/Public/Interfaces/OnlineFriendsInterface.h b/Engine/Plugins/Online/OnlineSubsystem/Source/Public/Interfaces/OnlineFriendsInterface.h
--- a/Engine/Plugins/Online/OnlineSubsystem/Source/Public/Interfaces/OnlineFriendsInterface.h
+++ b/Engine/Plugins/Online/OnlineSubsystem/Source/Public/Interfaces/OnlineFriendsInterface.h
@@ -82,15 +82,26 @@ struct FFriendSettings
 	TMap<FString, FString> SettingsMap;
 };
 
+/**
+ * Encounter type of recent player
+ */
+enum class ERecentPlayerEncounterType : uint8
+{
+	Default,
+	Teammate,
+	Opponent
+};
+
 /**
  * Stores information about a recent player
  */
 struct FReportPlayedWithUser
 {
 	FReportPlayedWithUser() = delete;
-	FReportPlayedWithUser(const FUniqueNetIdRef& InUserId, const FString& InPresenceStr)
+	FReportPlayedWithUser(const FUniqueNetIdRef& InUserId, const FString& InPresenceStr, ERecentPlayerEncounterType InEncounterType=ERecentPlayerEncounterType::Default)
 		: UserId(InUserId)
 		, PresenceStr(InPresenceStr)
+		, EncounterType(InEncounterType)
 	{
 	}
 
@@ -98,6 +109,8 @@ struct FReportPlayedWithUser
 	FUniqueNetIdRef UserId;
 	/** Optional presence string */
 	FString PresenceStr;
+	/** Encounter type of this user */
+	ERecentPlayerEncounterType EncounterType;
 };
 
 /**
