Move EXPOSE_USTRUCT_REP_DATA_PROPERTY_DIRECT to PartyTypes.h
Add method for getting escaped string size for values stored in an Online FVariantData

diff --git a/Engine/Plugins/Online/OnlineFramework/Source/Party/Public/Party/PartyTypes.h b/Engine/Plugins/Online/OnlineFramework/Source/Party/Public/Party/PartyTypes.h
--- a/Engine/Plugins/Online/OnlineFramework/Source/Party/Public/Party/PartyTypes.h
+++ b/Engine/Plugins/Online/OnlineFramework/Source/Party/Public/Party/PartyTypes.h
@@ -548,6 +548,10 @@ private: //
 EXPOSE_REP_DATA_PROPERTY_NO_SETTER(Owner, PropertyType, PropertyName, PropertyAccess);	\
 EXPOSE_REP_DATA_PROPERTY_SETTER_ONLY(Owner, PropertyType, PropertyName, PropertyAccess, SetterPrivacy)
 
+/** Expose ustruct property for access without using StructProperty in the property name. */
+#define EXPOSE_USTRUCT_REP_DATA_PROPERTY_DIRECT(Owner, PropertyType, StructProperty, ChildProperty)	\
+EXPOSE_REP_DATA_PROPERTY_SETTER_ACCESS(Owner, PropertyType, ChildProperty, StructProperty.ChildProperty, public);
+
  /** 
  * Simplest option - exposes getter and events, but no default setter
  * Revised version - allows the preservation of a previous property name for making delegates and accessors which will work along-side the new names.
diff --git a/Engine/Plugins/Online/OnlineSubsystem/Source/Private/OnlineKeyValuePair.cpp b/Engine/Plugins/Online/OnlineSubsystem/Source/Private/OnlineKeyValuePair.cpp
--- a/Engine/Plugins/Online/OnlineSubsystem/Source/Private/OnlineKeyValuePair.cpp
+++ b/Engine/Plugins/Online/OnlineSubsystem/Source/Private/OnlineKeyValuePair.cpp
@@ -1,6 +1,7 @@
 // Copyright Epic Games, Inc. All Rights Reserved.
 
 #include "OnlineKeyValuePair.h"
+#include "GenericPlatform/GenericPlatformMath.h"
 #include "Serialization/JsonTypes.h"
 #include "Dom/JsonObject.h"
 #include "UObject/ObjectMacros.h"
@@ -497,6 +498,120 @@ void FVariantData::GetValue(TArray<TSharedPtr<FJsonValue>>& OutData) const
 	}
 }
 
+
+int GetStringDigitCount(int64 Value)
+{
+	if (Value < 0)
+	{
+		const int SignLength = 1;
+		return FMath::CeilToInt64(FMath::LogX(10.f, -Value)) + SignLength;
+	}
+	else
+	{
+		return FMath::CeilToInt64(FMath::LogX(10.f, Value));
+	}
+}
+
+int GetStringDigitCount(uint64 Value)
+{
+	return FMath::CeilToInt64(FMath::LogX(10.f, Value));
+}
+
+int GetStringDigitCount(int32 Value)
+{
+	return GetStringDigitCount(static_cast<int64>(Value));
+}
+
+int GetStringDigitCount(uint32 Value)
+{
+	return GetStringDigitCount(static_cast<uint64>(Value));
+}
+
+/** Returns the maximum size of the value when encoded as an escaped string. */
+int FVariantData::GetEstimatedMaxEscapedStringSize() const
+{
+	if (CachedEstimatedMaxEscapedStringSize != 0)
+	{
+		return CachedEstimatedMaxEscapedStringSize;
+	}
+
+	switch (Type)
+	{
+		case EOnlineKeyValuePairDataType::Bool:
+		{
+			// Size of the strings "true" or "false"
+			CachedEstimatedMaxEscapedStringSize = 5;
+			break;
+		}
+
+		case EOnlineKeyValuePairDataType::Float:
+		{
+			// Maximum string size of a float value.
+			const int MaxDigits = 3 + FLT_MANT_DIG - FLT_MIN_EXP;
+			CachedEstimatedMaxEscapedStringSize = MaxDigits;
+			break;
+		}
+
+		case EOnlineKeyValuePairDataType::Double:
+		{
+			// Maximum string size of a double value.
+			const int MaxDigits = 3 + DBL_MANT_DIG - DBL_MIN_EXP;
+			CachedEstimatedMaxEscapedStringSize = MaxDigits;
+			break;
+		}
+
+		case EOnlineKeyValuePairDataType::Int32:
+		{
+			CachedEstimatedMaxEscapedStringSize = GetStringDigitCount(Value.AsInt);
+			break;
+		}
+
+		case EOnlineKeyValuePairDataType::UInt32:
+		{
+			CachedEstimatedMaxEscapedStringSize = GetStringDigitCount(Value.AsUInt);
+			break;
+		}
+
+		case EOnlineKeyValuePairDataType::Int64:
+		{
+			CachedEstimatedMaxEscapedStringSize = GetStringDigitCount(Value.AsInt64);
+			break;
+		}
+
+		case EOnlineKeyValuePairDataType::UInt64:
+		{
+			CachedEstimatedMaxEscapedStringSize = GetStringDigitCount(Value.AsUInt64);
+			break;
+		}
+
+		case EOnlineKeyValuePairDataType::Json:
+		case EOnlineKeyValuePairDataType::String:
+		{
+			check(Value.AsTCHAR);
+
+			// Get count of characters including extra padding for escaping sequences.
+			for (TCHAR* CharIt = Value.AsTCHAR; *CharIt; ++CharIt)
+			{
+				const int NumQuotes = *CharIt == TEXT('"');
+				const int NumBackSlash = (*CharIt == TEXT('\\')) * 2;
+				CachedEstimatedMaxEscapedStringSize += 1 + NumQuotes + NumBackSlash;
+			}
+			break;
+		}
+
+		case EOnlineKeyValuePairDataType::Empty:
+		case EOnlineKeyValuePairDataType::Blob:
+		default:
+		{
+			// Values are not exported to json.
+			CachedEstimatedMaxEscapedStringSize = 0;
+			break;
+		}
+	}
+
+	return CachedEstimatedMaxEscapedStringSize;
+}
+
 /**
  * Increments the numeric value by 1
  */
@@ -574,6 +689,7 @@ void FVariantData::Empty()
 
 	Type = EOnlineKeyValuePairDataType::Empty;
 	FMemory::Memset(&Value, 0, sizeof(ValueUnion));
+	CachedEstimatedMaxEscapedStringSize = 0;
 }
 
 /**
diff --git a/Engine/Plugins/Online/OnlineSubsystem/Source/Public/OnlineKeyValuePair.h b/Engine/Plugins/Online/OnlineSubsystem/Source/Public/OnlineKeyValuePair.h
--- a/Engine/Plugins/Online/OnlineSubsystem/Source/Public/OnlineKeyValuePair.h
+++ b/Engine/Plugins/Online/OnlineSubsystem/Source/Public/OnlineKeyValuePair.h
@@ -167,6 +167,9 @@ private:
 		ValueUnion() { FMemory::Memset( this, 0, sizeof( ValueUnion ) ); }
 	} Value;
 
+	/** Cached estimated maximum size of the value when encoded as an escaped string. */
+	mutable int CachedEstimatedMaxEscapedStringSize = 0;
+
 public:
 
 	/** Constructor */
@@ -410,6 +413,9 @@ public:
 	*/
 	void GetValue(TArray<TSharedPtr<class FJsonValue>>& OutData) const;
 
+	/** Returns the estimated maximum size of the value when encoded as an escaped string. */
+	int GetEstimatedMaxEscapedStringSize() const;
+
 	/**
 	 * Returns true if Type is numeric
 	 */
