Extend the OnlineGameActivity interface to allow for querying player's game activities

diff --git a/Engine/Plugins/Online/OnlineSubsystem/Source/Public/Interfaces/OnlineGameActivityInterface.h b/Engine/Plugins/Online/OnlineSubsystem/Source/Public/Interfaces/OnlineGameActivityInterface.h
--- a/Engine/Plugins/Online/OnlineSubsystem/Source/Public/Interfaces/OnlineGameActivityInterface.h
+++ b/Engine/Plugins/Online/OnlineSubsystem/Source/Public/Interfaces/OnlineGameActivityInterface.h
@@ -51,6 +51,32 @@ struct FOnlineActivityTasksToReset
 	TArray<FString> CompletedTasks;
 };
 
+/**
+ * Information about a game activity
+ */
+struct FOnlineActivityInformation
+{
+	/** Activity Id */
+	FString ActivityId;
+	/** Activity Name */
+	FString ActivityName;
+	/** Is the activity hidden */
+	bool bIsGameActivityHidden = false;
+	/** game match id */
+	FString GameMatchId;
+};
+
+/**
+ * Requested friend's activity information
+ */
+struct FUsersOnlineActivities
+{
+	/** The requesting account id */
+	FString AccountId;
+	/** Activity information */
+	TArray<FOnlineActivityInformation> Activities;
+};
+
 /** 
  * Multicast delegate fired when an activity request has happened
  *
@@ -113,6 +139,15 @@ DECLARE_DELEGATE_TwoParams(FOnSetActivityAvailabilityComplete, const FUniqueNetI
  */
 DECLARE_DELEGATE_TwoParams(FOnSetActivityPriorityComplete, const FUniqueNetId& /* LocalUserId */, const FOnlineError& /* Status */);
 
+/**
+ * Delegate fired when querying remote user activities call has completed
+ *
+ * @param LocalUserId the id of the player this callback is for
+ * @param UserActivities returned activity array for requested users
+ * @param Status of whether async action completed successfully or with error
+ */
+DECLARE_DELEGATE_ThreeParams(FOnQueryUserActivitiesComplete, const FUniqueNetId& /* LocalUserId */, const TArray<FUsersOnlineActivities>& /* UserActivities */, const FOnlineError& /* Status */);
+
 /**
  *	IOnlineGameActivity - Interface class for managing a user's activity state
  */
@@ -180,6 +215,16 @@ public:
 	 */
 	virtual void SetActivityPriority(const FUniqueNetId& LocalUserId, const TMap<FString, int32>& PrioritizedActivities, FOnSetActivityPriorityComplete CompletionDelegate) = 0;
 
+	/**
+	 * Query user's activities
+	 *
+	 * @param LocalUserId - Id of user requesting activities of other users
+	 * @param FriendAccountIds - Array of friend ids to get activity information on
+	 * @param MaxUserActivitiesReturned - Optional maximum number of activities to return per user
+	 * @param CompletionDelegate - Completion delegate called when QueryUserActivities call is complete
+	 */
+	virtual bool QueryActivitiesForUsers(const FUniqueNetId& LocalUserId, const TArray<FUniqueNetIdRef>& UserAccounts, TOptional<int32> MaxUserActivitiesReturned, const FOnQueryUserActivitiesComplete& CompletionDelegate) = 0;
+
 	/**
 	 * Called when an activity is requested
 	 *
