[Engine/Plugins]

diff --git a/Engine/Plugins/Online/OnlineFramework/Source/Party/Public/Chat/CommonSlashCommands.h b/Engine/Plugins/Online/OnlineFramework/Source/Party/Public/Chat/CommonSlashCommands.h
--- a/Engine/Plugins/Online/OnlineFramework/Source/Party/Public/Chat/CommonSlashCommands.h
+++ b/Engine/Plugins/Online/OnlineFramework/Source/Party/Public/Chat/CommonSlashCommands.h
@@ -3,9 +3,7 @@
 #pragma once
 
 #include "Chat/ChatSlashCommands.h"
-#include "Containers/UnrealString.h"
-#include "SocialTypes.h"
-#include "Chat/SocialChatChannel.h"
+#include "UObject/Object.h"
 
 class ULocalPlayer;
 class USocialUser;
@@ -70,3 +68,8 @@ private:
 };
 
 
+
+#if UE_ENABLE_INCLUDE_ORDER_DEPRECATED_IN_5_2
+#include "Chat/SocialChatChannel.h"
+#include "SocialTypes.h"
+#endif
diff --git a/Engine/Plugins/Online/OnlineFramework/Source/Party/Public/Chat/SocialChatChannel.h b/Engine/Plugins/Online/OnlineFramework/Source/Party/Public/Chat/SocialChatChannel.h
--- a/Engine/Plugins/Online/OnlineFramework/Source/Party/Public/Chat/SocialChatChannel.h
+++ b/Engine/Plugins/Online/OnlineFramework/Source/Party/Public/Chat/SocialChatChannel.h
@@ -2,12 +2,16 @@
 
 #pragma once
 
+#include "Interfaces/OnlineChatInterface.h"
 #include "OnlineSubsystem.h"
-#include "User/ISocialUserList.h"
-#include "User/SocialUser.h"
-#include "SocialChatMessage.h"
+#include "SocialTypes.h"
 #include "SocialChatChannel.generated.h"
 
+class FSocialLocalChatMessage;
+class USocialChatManager;
+class USocialToolkit;
+class USocialUser;
+
 UENUM(BlueprintType)
 enum class ESocialChannelType : uint8
 {
@@ -106,4 +110,10 @@ private:
 	mutable FOnMessageReceived OnMessageReceivedEvent;
 	mutable FOnChannelDisplayNameChanged OnChannelDisplayNameChangedEvent;
 
-};
\ No newline at end of file
+};
+
+#if UE_ENABLE_INCLUDE_ORDER_DEPRECATED_IN_5_2
+#include "SocialChatMessage.h"
+#include "User/ISocialUserList.h"
+#include "User/SocialUser.h"
+#endif
diff --git a/Engine/Plugins/Online/OnlineFramework/Source/Party/Public/Chat/SocialChatManager.h b/Engine/Plugins/Online/OnlineFramework/Source/Party/Public/Chat/SocialChatManager.h
--- a/Engine/Plugins/Online/OnlineFramework/Source/Party/Public/Chat/SocialChatManager.h
+++ b/Engine/Plugins/Online/OnlineFramework/Source/Party/Public/Chat/SocialChatManager.h
@@ -4,13 +4,9 @@
 
 #include "SocialChatChannel.h"
 #include "SocialToolkit.h"
-#include "OnlineSubsystem.h"
-#include "Interfaces/OnlineChatInterface.h"
 #include "SocialReadOnlyChatChannel.h"
 #include "SocialPrivateMessageChannel.h"
-#include "SocialChatRoom.h"
 #include "SocialGroupChannel.h"
-#include "Interfaces/OnlineGroupsInterface.h"
 #include "SocialChatManager.generated.h"
 
 class USocialChatRoom;
@@ -192,4 +188,8 @@ private:
 	mutable FOnChatChannelLeft OnChannelLeftEvent;
 	mutable FOnChatChannelFocusRequested OnChannelFocusRequestedEvent;
 	mutable FOnChatChannelDisplayRequested OnChannelDisplayRequestedEvent;
-};
\ No newline at end of file
+};
+
+#if UE_ENABLE_INCLUDE_ORDER_DEPRECATED_IN_5_2
+#include "SocialChatRoom.h"
+#endif
diff --git a/Engine/Plugins/Online/OnlineFramework/Source/Party/Public/Chat/SocialChatMessage.h b/Engine/Plugins/Online/OnlineFramework/Source/Party/Public/Chat/SocialChatMessage.h
--- a/Engine/Plugins/Online/OnlineFramework/Source/Party/Public/Chat/SocialChatMessage.h
+++ b/Engine/Plugins/Online/OnlineFramework/Source/Party/Public/Chat/SocialChatMessage.h
@@ -2,6 +2,7 @@
 
 #pragma once
 #include "Interfaces/OnlineChatInterface.h"
+#include "Misc/DateTime.h"
 #include "User/SocialUser.h"
 
 enum class ESocialChannelType : uint8;
@@ -172,4 +173,4 @@ protected:
 
 	FString SenderName;
 };
-#endif
\ No newline at end of file
+#endif
diff --git a/Engine/Plugins/Online/OnlineFramework/Source/Party/Public/Chat/SocialChatRoom.h b/Engine/Plugins/Online/OnlineFramework/Source/Party/Public/Chat/SocialChatRoom.h
--- a/Engine/Plugins/Online/OnlineFramework/Source/Party/Public/Chat/SocialChatRoom.h
+++ b/Engine/Plugins/Online/OnlineFramework/Source/Party/Public/Chat/SocialChatRoom.h
@@ -3,7 +3,6 @@
 #pragma once
 
 #include "SocialChatChannel.h"
-#include "Interfaces/OnlineChatInterface.h"
 #include "SocialChatRoom.generated.h"
 
 class USocialChatManager;
@@ -36,4 +35,4 @@ private:
 	void SetRoomId(const FChatRoomId& Id) { RoomId = Id; }
 
 	FChatRoomId RoomId;
-};
\ No newline at end of file
+};
diff --git a/Engine/Plugins/Online/OnlineFramework/Source/Party/Public/Chat/SocialGroupChannel.h b/Engine/Plugins/Online/OnlineFramework/Source/Party/Public/Chat/SocialGroupChannel.h
--- a/Engine/Plugins/Online/OnlineFramework/Source/Party/Public/Chat/SocialGroupChannel.h
+++ b/Engine/Plugins/Online/OnlineFramework/Source/Party/Public/Chat/SocialGroupChannel.h
@@ -2,9 +2,8 @@
 
 #pragma once
 
-#include "Chat/SocialChatChannel.h"
-#include "Online/CoreOnline.h"
 #include "Interfaces/OnlineGroupsInterface.h"
+#include "User/SocialUser.h"
 #include "SocialGroupChannel.generated.h"
 
 class USocialUser;
@@ -45,4 +44,8 @@ private:
 	TArray<TObjectPtr<USocialUser>> Members;
 
 	TWeakPtr<IOnlineGroups, ESPMode::ThreadSafe> GroupInterfacePtr;
-};
\ No newline at end of file
+};
+
+#if UE_ENABLE_INCLUDE_ORDER_DEPRECATED_IN_5_2
+#include "Chat/SocialChatChannel.h"
+#endif
diff --git a/Engine/Plugins/Online/OnlineFramework/Source/Party/Public/Chat/SocialPartyChatRoom.h b/Engine/Plugins/Online/OnlineFramework/Source/Party/Public/Chat/SocialPartyChatRoom.h
--- a/Engine/Plugins/Online/OnlineFramework/Source/Party/Public/Chat/SocialPartyChatRoom.h
+++ b/Engine/Plugins/Online/OnlineFramework/Source/Party/Public/Chat/SocialPartyChatRoom.h
@@ -5,6 +5,9 @@
 #include "SocialChatRoom.h"
 #include "SocialPartyChatRoom.generated.h"
 
+class UPartyMember;
+enum class ESocialChannelType : uint8;
+
 class USocialChatManager;
 enum class EMemberExitedReason : uint8;
 
@@ -20,4 +23,4 @@ public:
 private:
 	void HandlePartyMemberLeft(EMemberExitedReason Reason);
 	void HandlePartyMemberJoined(UPartyMember& NewMember);
-};
\ No newline at end of file
+};
diff --git a/Engine/Plugins/Online/OnlineFramework/Source/Party/Public/Chatroom.h b/Engine/Plugins/Online/OnlineFramework/Source/Party/Public/Chatroom.h
--- a/Engine/Plugins/Online/OnlineFramework/Source/Party/Public/Chatroom.h
+++ b/Engine/Plugins/Online/OnlineFramework/Source/Party/Public/Chatroom.h
@@ -1,10 +1,7 @@
 // Copyright Epic Games, Inc. All Rights Reserved.
 #pragma once
 
-#include "CoreMinimal.h"
-#include "UObject/ObjectMacros.h"
-#include "UObject/Object.h"
-#include "Engine/EngineTypes.h"
+#include "Engine/TimerHandle.h"
 #include "Interfaces/OnlineChatInterface.h"
 #include "Chatroom.generated.h"
 
@@ -133,3 +130,8 @@ protected:
 	UWorld* GetWorld() const;
 	FTimerManager& GetTimerManager() const;
 };
+
+#if UE_ENABLE_INCLUDE_ORDER_DEPRECATED_IN_5_2
+#include "CoreMinimal.h"
+#include "Engine/EngineTypes.h"
+#endif
diff --git a/Engine/Plugins/Online/OnlineFramework/Source/Party/Public/Interactions/CoreInteractions.h b/Engine/Plugins/Online/OnlineFramework/Source/Party/Public/Interactions/CoreInteractions.h
--- a/Engine/Plugins/Online/OnlineFramework/Source/Party/Public/Interactions/CoreInteractions.h
+++ b/Engine/Plugins/Online/OnlineFramework/Source/Party/Public/Interactions/CoreInteractions.h
@@ -2,6 +2,7 @@
 
 #pragma once
 #include "SocialInteractionMacros.h"
+#include "User/SocialUser.h"
 
 /**
  * Core set of user/friend interactions supported by the Social framework
@@ -18,4 +19,4 @@ DECLARE_SOCIAL_INTERACTION_EXPORT(PARTY_API, Block);
 DECLARE_SOCIAL_INTERACTION_EXPORT(PARTY_API, Unblock);
 DECLARE_SOCIAL_INTERACTION_EXPORT(PARTY_API, PrivateMessage);
 
-DECLARE_SOCIAL_INTERACTION_EXPORT(PARTY_API, ShowPlatformProfile);
\ No newline at end of file
+DECLARE_SOCIAL_INTERACTION_EXPORT(PARTY_API, ShowPlatformProfile);
diff --git a/Engine/Plugins/Online/OnlineFramework/Source/Party/Public/Interactions/PartyInteractions.h b/Engine/Plugins/Online/OnlineFramework/Source/Party/Public/Interactions/PartyInteractions.h
--- a/Engine/Plugins/Online/OnlineFramework/Source/Party/Public/Interactions/PartyInteractions.h
+++ b/Engine/Plugins/Online/OnlineFramework/Source/Party/Public/Interactions/PartyInteractions.h
@@ -2,6 +2,7 @@
 
 #pragma once
 #include "SocialInteractionMacros.h"
+#include "User/SocialUser.h"
 
 /**
 * Default Party interactions supported by the Social framework
@@ -19,4 +20,4 @@ DECLARE_SOCIAL_INTERACTION_EXPORT(PARTY_API, RejectPartyInvite);
 
 DECLARE_SOCIAL_INTERACTION_EXPORT(PARTY_API, LeaveParty);
 DECLARE_SOCIAL_INTERACTION_EXPORT(PARTY_API, KickPartyMember);
-DECLARE_SOCIAL_INTERACTION_EXPORT(PARTY_API, PromoteToPartyLeader);
\ No newline at end of file
+DECLARE_SOCIAL_INTERACTION_EXPORT(PARTY_API, PromoteToPartyLeader);
diff --git a/Engine/Plugins/Online/OnlineFramework/Source/Party/Public/Interactions/SocialInteractionHandle.h b/Engine/Plugins/Online/OnlineFramework/Source/Party/Public/Interactions/SocialInteractionHandle.h
--- a/Engine/Plugins/Online/OnlineFramework/Source/Party/Public/Interactions/SocialInteractionHandle.h
+++ b/Engine/Plugins/Online/OnlineFramework/Source/Party/Public/Interactions/SocialInteractionHandle.h
@@ -1,7 +1,10 @@
 // Copyright Epic Games, Inc. All Rights Reserved.
 
 #pragma once
-#include "CoreMinimal.h"
+#include "Containers/Map.h" // IWYU pragma: keep
+
+class FName;
+class FText;
 
 class USocialUser;
 class ISocialInteractionWrapper;
@@ -33,4 +36,8 @@ private:
 	FSocialInteractionHandle(const ISocialInteractionWrapper& Wrapper);
 
 	const ISocialInteractionWrapper* InteractionWrapper = nullptr;
-};
\ No newline at end of file
+};
+
+#if UE_ENABLE_INCLUDE_ORDER_DEPRECATED_IN_5_2
+#include "CoreMinimal.h"
+#endif
diff --git a/Engine/Plugins/Online/OnlineFramework/Source/Party/Public/Interactions/SocialInteractionMacros.h b/Engine/Plugins/Online/OnlineFramework/Source/Party/Public/Interactions/SocialInteractionMacros.h
--- a/Engine/Plugins/Online/OnlineFramework/Source/Party/Public/Interactions/SocialInteractionMacros.h
+++ b/Engine/Plugins/Online/OnlineFramework/Source/Party/Public/Interactions/SocialInteractionMacros.h
@@ -2,7 +2,8 @@
 
 #pragma once
 
-#include "User/SocialUser.h"
+#include "Interactions/SocialInteractionHandle.h"
+#include "Internationalization/Text.h"
 
 DECLARE_DELEGATE_RetVal_OneParam(bool, FOnCustomIsInteractionAvailable, const USocialUser&);
 
@@ -91,3 +92,7 @@ private: \
 }
 
 #define DECLARE_SOCIAL_INTERACTION(InteractionName) DECLARE_SOCIAL_INTERACTION_EXPORT(, InteractionName)
+
+#if UE_ENABLE_INCLUDE_ORDER_DEPRECATED_IN_5_2
+#include "User/SocialUser.h"
+#endif
diff --git a/Engine/Plugins/Online/OnlineFramework/Source/Party/Public/Party/PartyDataReplicator.h b/Engine/Plugins/Online/OnlineFramework/Source/Party/Public/Party/PartyDataReplicator.h
--- a/Engine/Plugins/Online/OnlineFramework/Source/Party/Public/Party/PartyDataReplicator.h
+++ b/Engine/Plugins/Online/OnlineFramework/Source/Party/Public/Party/PartyDataReplicator.h
@@ -3,10 +3,10 @@
 #pragma once
 
 #include "Party/PartyTypes.h"
+#include "PartyModule.h"
 #include "UObject/GCObject.h"
 #include "Containers/Ticker.h"
 #include "Interfaces/OnlinePartyInterface.h"
-#include "Stats/Stats.h"
 
 /** Util exclusively for use by TPartyDataReplicator to circumvent circular include header issues (we can't include SocialParty.h or PartyMember.h here) */
 class FPartyDataReplicatorHelper
@@ -155,4 +155,4 @@ private:
 	RepDataT* RepDataCopy = nullptr;
 
 	FTSTicker::FDelegateHandle UpdateTickerHandle;
-};
\ No newline at end of file
+};
diff --git a/Engine/Plugins/Online/OnlineFramework/Source/Party/Public/Party/PartyMember.h b/Engine/Plugins/Online/OnlineFramework/Source/Party/Public/Party/PartyMember.h
--- a/Engine/Plugins/Online/OnlineFramework/Source/Party/Public/Party/PartyMember.h
+++ b/Engine/Plugins/Online/OnlineFramework/Source/Party/Public/Party/PartyMember.h
@@ -2,9 +2,9 @@
 
 #pragma once
 
-#include "Party/PartyTypes.h"
 #include "Party/PartyDataReplicator.h"
 
+#include "SocialTypes.h"
 #include "PartyMember.generated.h"
 
 class USocialUser;
@@ -243,4 +243,4 @@ private:
 	mutable FOnPartyMemberStateChanged OnPromotedToLeaderEvent;
 	mutable FOnPartyMemberStateChanged OnDemotedEvent;
 	mutable FOnPartyMemberLeft OnLeftPartyEvent;
-};
\ No newline at end of file
+};
diff --git a/Engine/Plugins/Online/OnlineFramework/Source/Party/Public/Party/PartyTypes.h b/Engine/Plugins/Online/OnlineFramework/Source/Party/Public/Party/PartyTypes.h
--- a/Engine/Plugins/Online/OnlineFramework/Source/Party/Public/Party/PartyTypes.h
+++ b/Engine/Plugins/Online/OnlineFramework/Source/Party/Public/Party/PartyTypes.h
@@ -2,8 +2,8 @@
 
 #pragma once
 
-#include "SocialTypes.h"
 
+#include "GameFramework/OnlineReplStructs.h"
 #include "PartyTypes.generated.h"
 
 class USocialParty;
@@ -700,3 +700,7 @@ namespace PartyJoinMethod
 	const FName Unspecified = FName(TEXT("Unspecified"));
 	
 }
+
+#if UE_ENABLE_INCLUDE_ORDER_DEPRECATED_IN_5_2
+#include "SocialTypes.h"
+#endif
diff --git a/Engine/Plugins/Online/OnlineFramework/Source/Party/Public/Party/SocialParty.h b/Engine/Plugins/Online/OnlineFramework/Source/Party/Public/Party/SocialParty.h
--- a/Engine/Plugins/Online/OnlineFramework/Source/Party/Public/Party/SocialParty.h
+++ b/Engine/Plugins/Online/OnlineFramework/Source/Party/Public/Party/SocialParty.h
@@ -2,18 +2,20 @@
 
 #pragma once
 
-#include "Party/PartyTypes.h"
 #include "Party/PartyDataReplicator.h"
 
-#include "PartyBeaconState.h"
-#include "SpectatorBeaconState.h"
+#include "SocialTypes.h"
 #include "SpectatorBeaconClient.h"
-#include "Interfaces/OnlinePartyInterface.h"
 #include "Interfaces/OnlineChatInterface.h"
 #include "Containers/Queue.h"
-#include "Engine/EngineBaseTypes.h"
 #include "SocialParty.generated.h"
 
+class AOnlineBeaconClient;
+class FOnlinePartyId;
+class FOnlinePartyTypeId;
+enum ETravelType : int;
+namespace EPartyReservationResult { enum Type : int; }
+
 class APartyBeaconClient;
 class UNetDriver;
 
@@ -495,3 +497,7 @@ private:
 	mutable FOnPartyJIPApproved OnPartyJIPApprovedEvent;
 	mutable FOnPartyJIPResponse OnPartyJIPResponseEvent;
 };
+
+#if UE_ENABLE_INCLUDE_ORDER_DEPRECATED_IN_5_2
+#include "PartyBeaconState.h"
+#endif
diff --git a/Engine/Plugins/Online/OnlineFramework/Source/Party/Public/PartyModule.h b/Engine/Plugins/Online/OnlineFramework/Source/Party/Public/PartyModule.h
--- a/Engine/Plugins/Online/OnlineFramework/Source/Party/Public/PartyModule.h
+++ b/Engine/Plugins/Online/OnlineFramework/Source/Party/Public/PartyModule.h
@@ -2,10 +2,7 @@
 
 #pragma once
 
-#include "CoreMinimal.h"
 #include "Stats/Stats.h"
-#include "Misc/CoreMisc.h"
-#include "Modules/ModuleInterface.h"
 #include "Modules/ModuleManager.h"
 
 /** Logging related to parties */
@@ -63,3 +60,7 @@ private:
 	virtual void ShutdownModule() override;
 };
 
+
+#if UE_ENABLE_INCLUDE_ORDER_DEPRECATED_IN_5_2
+#include "CoreMinimal.h"
+#endif
diff --git a/Engine/Plugins/Online/OnlineFramework/Source/Party/Public/SocialDebugTools.h b/Engine/Plugins/Online/OnlineFramework/Source/Party/Public/SocialDebugTools.h
--- a/Engine/Plugins/Online/OnlineFramework/Source/Party/Public/SocialDebugTools.h
+++ b/Engine/Plugins/Online/OnlineFramework/Source/Party/Public/SocialDebugTools.h
@@ -2,10 +2,13 @@
 
 #pragma once
 
-#include "UObject/Object.h"
-#include "SocialTypes.h"
 #include "SocialDebugTools.generated.h"
 
+class FOnlinePartyId;
+class FUniqueNetId;
+class IOnlinePartyJoinInfo;
+class USocialManager;
+
 class IOnlineSubsystem;
 class FOnlineAccountCredentials;
 class IOnlinePartyPendingJoinRequestInfo;
@@ -102,4 +105,8 @@ private:
 	void HandleFriendInviteReceived(const FUniqueNetId& LocalUserId, const FUniqueNetId& FriendId);
 	void HandlePartyInviteReceived(const FUniqueNetId& LocalUserId, const IOnlinePartyJoinInfo& Invitation);
 	void HandlePartyJoinRequestReceived(const FUniqueNetId& LocalUserId, const FOnlinePartyId& PartyId, const IOnlinePartyPendingJoinRequestInfo& JoinRequestInfo);
-};
\ No newline at end of file
+};
+
+#if UE_ENABLE_INCLUDE_ORDER_DEPRECATED_IN_5_2
+#include "SocialTypes.h"
+#endif
diff --git a/Engine/Plugins/Online/OnlineFramework/Source/Party/Public/SocialManager.h b/Engine/Plugins/Online/OnlineFramework/Source/Party/Public/SocialManager.h
--- a/Engine/Plugins/Online/OnlineFramework/Source/Party/Public/SocialManager.h
+++ b/Engine/Plugins/Online/OnlineFramework/Source/Party/Public/SocialManager.h
@@ -2,17 +2,18 @@
 
 #pragma once
 
-#include "UObject/Object.h"
-#include "OnlineSubsystem.h"
-#include "OnlineSubsystemTypes.h"
-#include "Engine/EngineBaseTypes.h"
 
 #include "Party/PartyTypes.h"
 #include "Interfaces/OnlinePartyInterface.h"
-#include "Interactions/SocialInteractionHandle.h"
 
+#include "SocialTypes.h"
+#include "Templates/SubclassOf.h"
 #include "SocialManager.generated.h"
 
+class FSocialInteractionHandle;
+class IOnlineSubsystem;
+struct FOnlineError;
+
 class ULocalPlayer;
 class USocialUser;
 class USocialParty;
@@ -289,4 +290,10 @@ private:
 
 	mutable FOnSocialToolkitCreated OnSocialToolkitCreatedEvent;
 	mutable FOnPartyMembershipChanged OnPartyJoinedEvent;
-};
\ No newline at end of file
+};
+
+#if UE_ENABLE_INCLUDE_ORDER_DEPRECATED_IN_5_2
+#include "Engine/EngineBaseTypes.h"
+#include "Interactions/SocialInteractionHandle.h"
+#include "OnlineSubsystem.h"
+#endif
diff --git a/Engine/Plugins/Online/OnlineFramework/Source/Party/Public/SocialQuery.h b/Engine/Plugins/Online/OnlineFramework/Source/Party/Public/SocialQuery.h
--- a/Engine/Plugins/Online/OnlineFramework/Source/Party/Public/SocialQuery.h
+++ b/Engine/Plugins/Online/OnlineFramework/Source/Party/Public/SocialQuery.h
@@ -1,12 +1,8 @@
 // Copyright Epic Games, Inc. All Rights Reserved.
 #pragma once
 
-#include "CoreMinimal.h"
-#include "Containers/Ticker.h"
-#include "OnlineSubsystem.h"
 #include "SocialToolkit.h"
 #include "Misc/ConfigCacheIni.h"
-#include "Stats/Stats.h"
 
 DECLARE_DELEGATE_TwoParams(FOnQueryCompleted, FName, const TSharedRef<class FSocialQueryBase>&);
 
@@ -163,4 +159,9 @@ private:
 
 	FTSTicker::FDelegateHandle TickExecuteHandle;
 	TMap<FName, TArray<TSharedRef<FSocialQueryBase>>> CurrentQueriesById;
-};
\ No newline at end of file
+};
+
+#if UE_ENABLE_INCLUDE_ORDER_DEPRECATED_IN_5_2
+#include "CoreMinimal.h"
+#include "OnlineSubsystem.h"
+#endif
diff --git a/Engine/Plugins/Online/OnlineFramework/Source/Party/Public/SocialSettings.h b/Engine/Plugins/Online/OnlineFramework/Source/Party/Public/SocialSettings.h
--- a/Engine/Plugins/Online/OnlineFramework/Source/Party/Public/SocialSettings.h
+++ b/Engine/Plugins/Online/OnlineFramework/Source/Party/Public/SocialSettings.h
@@ -2,7 +2,6 @@
 
 #pragma once
 
-#include "UObject/Object.h"
 
 #include "SocialSettings.generated.h"
 
@@ -123,4 +122,4 @@ private:
 
 	UPROPERTY(Config)
 	TArray<FName> SonyOSSNames;
-};
\ No newline at end of file
+};
diff --git a/Engine/Plugins/Online/OnlineFramework/Source/Party/Public/SocialToolkit.h b/Engine/Plugins/Online/OnlineFramework/Source/Party/Public/SocialToolkit.h
--- a/Engine/Plugins/Online/OnlineFramework/Source/Party/Public/SocialToolkit.h
+++ b/Engine/Plugins/Online/OnlineFramework/Source/Party/Public/SocialToolkit.h
@@ -2,16 +2,19 @@
 
 #pragma once
 
-#include "UObject/Object.h"
-#include "SocialTypes.h"
-#include "User/ISocialUserList.h"
-#include "OnlineSubsystem.h"
 #include "Interfaces/OnlinePartyInterface.h"
 #include "Interfaces/OnlinePresenceInterface.h"
 #include "Containers/Ticker.h"
 
+#include "PartyModule.h"
+#include "Templates/SubclassOf.h"
+#include "User/SocialUser.h"
 #include "SocialToolkit.generated.h"
 
+class FSocialUserListConfig;
+class IOnlineSubsystem;
+class USocialManager;
+
 class ULocalPlayer;
 class USocialUser;
 class FSocialFriendInvitation;
@@ -335,3 +338,8 @@ private:
 	mutable FOnPartyRequestToJoinReceivedEvent OnPartyRequestToJoinReceivedEvent;
 	mutable FOnPartyRequestToJoinRemovedEvent OnPartyRequestToJoinRemovedEvent;
 };
+
+#if UE_ENABLE_INCLUDE_ORDER_DEPRECATED_IN_5_2
+#include "OnlineSubsystem.h"
+#include "User/ISocialUserList.h"
+#endif
diff --git a/Engine/Plugins/Online/OnlineFramework/Source/Party/Public/SocialTypes.h b/Engine/Plugins/Online/OnlineFramework/Source/Party/Public/SocialTypes.h
--- a/Engine/Plugins/Online/OnlineFramework/Source/Party/Public/SocialTypes.h
+++ b/Engine/Plugins/Online/OnlineFramework/Source/Party/Public/SocialTypes.h
@@ -2,11 +2,6 @@
 
 #pragma once
 
-#include "PartyModule.h"
-#include "OnlineSubsystemTypes.h"
-#include "GameFramework/OnlineReplStructs.h"
-#include "Templates/SubclassOf.h"
-#include "Interactions/SocialInteractionHandle.h"
 #include "SocialSettings.h"
 
 #include "SocialTypes.generated.h"
@@ -195,4 +190,12 @@ inline const TCHAR* LexToString(ECrossplayPreference Preference)
 	case ECrossplayPreference::OptedOut: return TEXT("OptedOut");
 	}
 	return TEXT("Unknown");
-}
\ No newline at end of file
+}
+
+#if UE_ENABLE_INCLUDE_ORDER_DEPRECATED_IN_5_2
+#include "GameFramework/OnlineReplStructs.h"
+#include "Interactions/SocialInteractionHandle.h"
+#include "OnlineSubsystemTypes.h"
+#include "PartyModule.h"
+#include "Templates/SubclassOf.h"
+#endif
diff --git a/Engine/Plugins/Online/OnlineFramework/Source/Party/Public/User/ISocialUserList.h b/Engine/Plugins/Online/OnlineFramework/Source/Party/Public/User/ISocialUserList.h
--- a/Engine/Plugins/Online/OnlineFramework/Source/Party/Public/User/ISocialUserList.h
+++ b/Engine/Plugins/Online/OnlineFramework/Source/Party/Public/User/ISocialUserList.h
@@ -2,9 +2,10 @@
 
 #pragma once
 
+#include "PartyModule.h"
 #include "SocialTypes.h"
-#include "Misc/EnumClassFlags.h"
-#include "SocialUser.h"
+
+class USocialUser;
 
 DECLARE_DELEGATE_RetVal_OneParam(bool, FOnCustomFilterUser, const USocialUser&);
 
@@ -101,4 +102,8 @@ public:
 	virtual void SetAllowAutoUpdate(bool bIsEnabled) = 0;
 
 	virtual void SetAllowSortDuringUpdate(bool bIsEnabled) = 0;
-};
\ No newline at end of file
+};
+
+#if UE_ENABLE_INCLUDE_ORDER_DEPRECATED_IN_5_2
+#include "SocialUser.h"
+#endif
diff --git a/Engine/Plugins/Online/OnlineFramework/Source/Party/Public/User/SocialUser.h b/Engine/Plugins/Online/OnlineFramework/Source/Party/Public/User/SocialUser.h
--- a/Engine/Plugins/Online/OnlineFramework/Source/Party/Public/User/SocialUser.h
+++ b/Engine/Plugins/Online/OnlineFramework/Source/Party/Public/User/SocialUser.h
@@ -2,14 +2,16 @@
 
 #pragma once
 
-#include "Interactions/SocialInteractionHandle.h"
-#include "OnlineSessionSettings.h"
+#include "OnlineSubsystemTypes.h"
 #include "Party/PartyTypes.h"
 #include "SocialTypes.h"
-#include "UObject/Object.h"
 
 #include "SocialUser.generated.h"
 
+class FSocialInteractionHandle;
+class USocialToolkit;
+class USocialUser;
+
 class IOnlinePartyJoinInfo;
 class FOnlineUserPresence;
 class UPartyMember;
@@ -262,3 +264,8 @@ private:
 	mutable FOnSubsystemIdEstablished OnSubsystemIdEstablishedEvent;
 	mutable FOnUserGameSpecificStatusChanged OnUserGameSpecificStatusChangedEvent;
 };
+
+#if UE_ENABLE_INCLUDE_ORDER_DEPRECATED_IN_5_2
+#include "Interactions/SocialInteractionHandle.h"
+#include "OnlineSessionSettings.h"
+#endif
diff --git a/Engine/Plugins/Online/OnlineFramework/Source/PlayTimeLimit/Public/PlayTimeLimitModule.h b/Engine/Plugins/Online/OnlineFramework/Source/PlayTimeLimit/Public/PlayTimeLimitModule.h
--- a/Engine/Plugins/Online/OnlineFramework/Source/PlayTimeLimit/Public/PlayTimeLimitModule.h
+++ b/Engine/Plugins/Online/OnlineFramework/Source/PlayTimeLimit/Public/PlayTimeLimitModule.h
@@ -2,10 +2,6 @@
 
 #pragma once
 
-#include "CoreMinimal.h"
-#include "Stats/Stats.h"
-#include "Misc/CoreMisc.h"
-#include "Modules/ModuleInterface.h"
 #include "Modules/ModuleManager.h"
 
 PLAYTIMELIMIT_API DECLARE_LOG_CATEGORY_EXTERN(LogPlayTimeLimit, Log, All);
@@ -62,3 +58,8 @@ private:
 	bool bPlayTimeLimitEnabled = false;
 };
 
+
+#if UE_ENABLE_INCLUDE_ORDER_DEPRECATED_IN_5_2
+#include "CoreMinimal.h"
+#include "Stats/Stats.h"
+#endif
diff --git a/Engine/Plugins/Online/OnlineFramework/Source/PlayTimeLimit/Public/PlayTimeLimitUser.h b/Engine/Plugins/Online/OnlineFramework/Source/PlayTimeLimit/Public/PlayTimeLimitUser.h
--- a/Engine/Plugins/Online/OnlineFramework/Source/PlayTimeLimit/Public/PlayTimeLimitUser.h
+++ b/Engine/Plugins/Online/OnlineFramework/Source/PlayTimeLimit/Public/PlayTimeLimitUser.h
@@ -2,8 +2,7 @@
 
 #pragma once
 
-#include "CoreMinimal.h"
-#include "Online/CoreOnline.h"
+#include "Online/CoreOnlineFwd.h"
 
 /**
  * Information about a user we are observing the play time and instituting limits for
@@ -106,3 +105,8 @@ protected:
 
 typedef FPlayTimeLimitUser* FPlayTimeLimitUserRawPtr;
 typedef TSharedPtr<FPlayTimeLimitUser, ESPMode::ThreadSafe> FPlayTimeLimitUserPtr;
+
+#if UE_ENABLE_INCLUDE_ORDER_DEPRECATED_IN_5_2
+#include "CoreMinimal.h"
+#include "Online/CoreOnline.h"
+#endif
diff --git a/Engine/Plugins/Online/OnlineFramework/Source/Qos/Public/QosInterface.h b/Engine/Plugins/Online/OnlineFramework/Source/Qos/Public/QosInterface.h
--- a/Engine/Plugins/Online/OnlineFramework/Source/Qos/Public/QosInterface.h
+++ b/Engine/Plugins/Online/OnlineFramework/Source/Qos/Public/QosInterface.h
@@ -2,10 +2,11 @@
 
 #pragma once
 
-#include "CoreMinimal.h"
-#include "QosRegionManager.h"
 #include "UObject/GCObject.h"
 
+class UQosRegionManager;
+struct FRegionQosInstance;
+
 class IAnalyticsProvider;
 
 #define NO_REGION TEXT("NONE")
@@ -148,3 +149,8 @@ private:
 	UQosRegionManager* RegionManager;
 };
 
+
+#if UE_ENABLE_INCLUDE_ORDER_DEPRECATED_IN_5_2
+#include "CoreMinimal.h"
+#include "QosRegionManager.h"
+#endif
diff --git a/Engine/Plugins/Online/OnlineFramework/Source/Qos/Public/QosModule.h b/Engine/Plugins/Online/OnlineFramework/Source/Qos/Public/QosModule.h
--- a/Engine/Plugins/Online/OnlineFramework/Source/Qos/Public/QosModule.h
+++ b/Engine/Plugins/Online/OnlineFramework/Source/Qos/Public/QosModule.h
@@ -2,9 +2,6 @@
 
 #pragma once
 
-#include "CoreMinimal.h"
-#include "Misc/CoreMisc.h"
-#include "Modules/ModuleInterface.h"
 #include "Modules/ModuleManager.h"
 
 class FQosInterface;
@@ -70,3 +67,7 @@ private:
 	virtual void ShutdownModule() override;
 };
 
+
+#if UE_ENABLE_INCLUDE_ORDER_DEPRECATED_IN_5_2
+#include "CoreMinimal.h"
+#endif
diff --git a/Engine/Plugins/Online/OnlineFramework/Source/Qos/Public/QosRegionManager.h b/Engine/Plugins/Online/OnlineFramework/Source/Qos/Public/QosRegionManager.h
--- a/Engine/Plugins/Online/OnlineFramework/Source/Qos/Public/QosRegionManager.h
+++ b/Engine/Plugins/Online/OnlineFramework/Source/Qos/Public/QosRegionManager.h
@@ -2,10 +2,8 @@
 
 #pragma once
 
-#include "CoreMinimal.h"
-#include "UObject/ObjectMacros.h"
-#include "UObject/Object.h"
 
+#include "Misc/DateTime.h"
 #include "QosRegionManager.generated.h"
 
 class IAnalyticsProvider;
@@ -700,3 +698,7 @@ private:
 	static const TCHAR* SubspaceDelimiterDefault;
 };
 
+
+#if UE_ENABLE_INCLUDE_ORDER_DEPRECATED_IN_5_2
+#include "CoreMinimal.h"
+#endif
diff --git a/Engine/Plugins/Online/OnlineFramework/Source/Rejoin/Public/RejoinCheck.h b/Engine/Plugins/Online/OnlineFramework/Source/Rejoin/Public/RejoinCheck.h
--- a/Engine/Plugins/Online/OnlineFramework/Source/Rejoin/Public/RejoinCheck.h
+++ b/Engine/Plugins/Online/OnlineFramework/Source/Rejoin/Public/RejoinCheck.h
@@ -1,11 +1,8 @@
 // Copyright Epic Games, Inc. All Rights Reserved.
 #pragma once
 
-#include "Engine/EngineTypes.h"
+#include "Engine/TimerHandle.h"
 #include "OnlineSessionSettings.h"
-#include "UObject/ObjectMacros.h"
-#include "UObject/Object.h"
-#include "Engine/EngineTypes.h"
 #include "RejoinCheck.generated.h"
 
 class FTimerManager;
@@ -325,3 +322,7 @@ inline const TCHAR* ToString(ERejoinAttemptResult Result)
 	return TEXT("");
 }
 
+
+#if UE_ENABLE_INCLUDE_ORDER_DEPRECATED_IN_5_2
+#include "Engine/EngineTypes.h"
+#endif
diff --git a/Engine/Plugins/Online/OnlineFramework/Source/Rejoin/Public/RejoinModule.h b/Engine/Plugins/Online/OnlineFramework/Source/Rejoin/Public/RejoinModule.h
--- a/Engine/Plugins/Online/OnlineFramework/Source/Rejoin/Public/RejoinModule.h
+++ b/Engine/Plugins/Online/OnlineFramework/Source/Rejoin/Public/RejoinModule.h
@@ -2,7 +2,7 @@
 
 #pragma once
 
-#include "Modules/ModuleManager.h"
+#include "Modules/ModuleInterface.h"
 
 /**
  * Module for Rejoin base implementation
@@ -15,3 +15,7 @@ private:
 	virtual void StartupModule() override;
 	virtual void ShutdownModule() override;
 };
+
+#if UE_ENABLE_INCLUDE_ORDER_DEPRECATED_IN_5_2
+#include "Modules/ModuleManager.h"
+#endif
diff --git a/Engine/Plugins/Online/OnlineServices/Source/OnlineServicesCommon/Private/Online/OnlineAsyncOpCache.cpp b/Engine/Plugins/Online/OnlineServices/Source/OnlineServicesCommon/Private/Online/OnlineAsyncOpCache.cpp
--- a/Engine/Plugins/Online/OnlineServices/Source/OnlineServicesCommon/Private/Online/OnlineAsyncOpCache.cpp
+++ b/Engine/Plugins/Online/OnlineServices/Source/OnlineServicesCommon/Private/Online/OnlineAsyncOpCache.cpp
@@ -1,6 +1,7 @@
 // Copyright Epic Games, Inc. All Rights Reserved.
 
 #include "Online/OnlineAsyncOpCache.h"
+#include "Online/OnlineExecHandler.h" // IWYU pragma: keep
 #include "Online/OnlineServicesCommon.h"
 
 namespace UE::Online {
diff --git a/Engine/Plugins/Online/OnlineServices/Source/OnlineServicesCommon/Public/Online/AuthCommon.h b/Engine/Plugins/Online/OnlineServices/Source/OnlineServicesCommon/Public/Online/AuthCommon.h
--- a/Engine/Plugins/Online/OnlineServices/Source/OnlineServicesCommon/Public/Online/AuthCommon.h
+++ b/Engine/Plugins/Online/OnlineServices/Source/OnlineServicesCommon/Public/Online/AuthCommon.h
@@ -2,7 +2,6 @@
 
 #pragma once
 
-#include "Online/Auth.h"
 #include "Online/OnlineComponent.h"
 
 namespace UE::Online {
diff --git a/Engine/Plugins/Online/OnlineServices/Source/OnlineServicesCommon/Public/Online/IOnlineComponent.h b/Engine/Plugins/Online/OnlineServices/Source/OnlineServicesCommon/Public/Online/IOnlineComponent.h
--- a/Engine/Plugins/Online/OnlineServices/Source/OnlineServicesCommon/Public/Online/IOnlineComponent.h
+++ b/Engine/Plugins/Online/OnlineServices/Source/OnlineServicesCommon/Public/Online/IOnlineComponent.h
@@ -2,8 +2,6 @@
 
 #pragma once
 
-#include "Templates/SharedPointer.h"
-#include "Online/OnlineAsyncOp.h"
 
 namespace UE::Online {
 
@@ -26,3 +24,8 @@ public:
 };
 
 /* UE::Online */ }
+
+#if UE_ENABLE_INCLUDE_ORDER_DEPRECATED_IN_5_2
+#include "Online/OnlineAsyncOp.h"
+#include "Templates/SharedPointer.h"
+#endif
diff --git a/Engine/Plugins/Online/OnlineServices/Source/OnlineServicesCommon/Public/Online/LobbiesCommon.h b/Engine/Plugins/Online/OnlineServices/Source/OnlineServicesCommon/Public/Online/LobbiesCommon.h
--- a/Engine/Plugins/Online/OnlineServices/Source/OnlineServicesCommon/Public/Online/LobbiesCommon.h
+++ b/Engine/Plugins/Online/OnlineServices/Source/OnlineServicesCommon/Public/Online/LobbiesCommon.h
@@ -2,7 +2,6 @@
 
 #pragma once
 
-#include "Online/Lobbies.h"
 #include "Online/OnlineComponent.h"
 #include "Online/LobbiesCommonTypes.h"
 
diff --git a/Engine/Plugins/Online/OnlineServices/Source/OnlineServicesCommon/Public/Online/LobbiesCommonTypes.h b/Engine/Plugins/Online/OnlineServices/Source/OnlineServicesCommon/Public/Online/LobbiesCommonTypes.h
--- a/Engine/Plugins/Online/OnlineServices/Source/OnlineServicesCommon/Public/Online/LobbiesCommonTypes.h
+++ b/Engine/Plugins/Online/OnlineServices/Source/OnlineServicesCommon/Public/Online/LobbiesCommonTypes.h
@@ -4,6 +4,7 @@
 
 #include "Async/Future.h"
 #include "Online/Lobbies.h"
+#include "Online/OnlineAsyncOpHandle.h"
 #include "Online/Schema.h"
 
 namespace UE::Online {
diff --git a/Engine/Plugins/Online/OnlineServices/Source/OnlineServicesCommon/Public/Online/NboSerializerCommonSvc.h b/Engine/Plugins/Online/OnlineServices/Source/OnlineServicesCommon/Public/Online/NboSerializerCommonSvc.h
--- a/Engine/Plugins/Online/OnlineServices/Source/OnlineServicesCommon/Public/Online/NboSerializerCommonSvc.h
+++ b/Engine/Plugins/Online/OnlineServices/Source/OnlineServicesCommon/Public/Online/NboSerializerCommonSvc.h
@@ -2,9 +2,13 @@
 
 #pragma once
 
-#include "Online/NboSerializer.h"
 
-#include "Online/SessionsCommon.h"
+#include "Online/Sessions.h"
+
+class FNboSerializeFromBuffer;
+class FNboSerializeToBuffer;
+namespace UE::Online { class FSchemaVariant; }
+namespace UE::Online { class FSessionCommon; }
 
 /**
  * Serializes data in network byte order form into a buffer
@@ -31,4 +35,9 @@ ONLINESERVICESCOMMON_API void SerializeFromBuffer(FNboSerializeFromBuffer& Packe
 
 /* NboSerializerCommonSvc */ }
 
-/* UE::Online */ }
\ No newline at end of file
+/* UE::Online */ }
+
+#if UE_ENABLE_INCLUDE_ORDER_DEPRECATED_IN_5_2
+#include "Online/NboSerializer.h"
+#include "Online/SessionsCommon.h"
+#endif
diff --git a/Engine/Plugins/Online/OnlineServices/Source/OnlineServicesCommon/Public/Online/OnlineAsyncOp.h b/Engine/Plugins/Online/OnlineServices/Source/OnlineServicesCommon/Public/Online/OnlineAsyncOp.h
--- a/Engine/Plugins/Online/OnlineServices/Source/OnlineServicesCommon/Public/Online/OnlineAsyncOp.h
+++ b/Engine/Plugins/Online/OnlineServices/Source/OnlineServicesCommon/Public/Online/OnlineAsyncOp.h
@@ -3,11 +3,8 @@
 #pragma once
 
 #include "Online/OnlineAsyncOpHandle.h"
+#include "Online/OnlineResult.h"
 #include "Online/OnlineTypeInfo.h"
-#include "Online/OnlineErrorDefinitions.h"
-#include "Containers/Map.h"
-#include "Templates/UniquePtr.h"
-#include "Async/Future.h"
 #include "Async/Async.h"
 
 namespace UE::Online {
diff --git a/Engine/Plugins/Online/OnlineServices/Source/OnlineServicesCommon/Public/Online/OnlineAsyncOpCache.h b/Engine/Plugins/Online/OnlineServices/Source/OnlineServicesCommon/Public/Online/OnlineAsyncOpCache.h
--- a/Engine/Plugins/Online/OnlineServices/Source/OnlineServicesCommon/Public/Online/OnlineAsyncOpCache.h
+++ b/Engine/Plugins/Online/OnlineServices/Source/OnlineServicesCommon/Public/Online/OnlineAsyncOpCache.h
@@ -4,8 +4,7 @@
 
 #include "Online/CoreOnline.h"
 #include "Online/OnlineAsyncOp.h"
-#include "Containers/Map.h"
-#include "Templates/SharedPointer.h"
+#include "Online/OnlineMeta.h" // IWYU pragma: keep
 
 namespace UE::Online {
 
diff --git a/Engine/Plugins/Online/OnlineServices/Source/OnlineServicesCommon/Public/Online/OnlineAsyncOpCache_Meta.inl b/Engine/Plugins/Online/OnlineServices/Source/OnlineServicesCommon/Public/Online/OnlineAsyncOpCache_Meta.inl
--- a/Engine/Plugins/Online/OnlineServices/Source/OnlineServicesCommon/Public/Online/OnlineAsyncOpCache_Meta.inl
+++ b/Engine/Plugins/Online/OnlineServices/Source/OnlineServicesCommon/Public/Online/OnlineAsyncOpCache_Meta.inl
@@ -4,6 +4,9 @@
 
 #include "Online/OnlineMeta.h"
 
+namespace UE::Online { struct FOperationConfig; }
+namespace UE::Online::Meta { template <typename StructType> struct TStructDetails; }
+
 namespace UE::Online::Meta {
 
 BEGIN_ONLINE_STRUCT_META(FOperationConfig)
diff --git a/Engine/Plugins/Online/OnlineServices/Source/OnlineServicesCommon/Public/Online/OnlineAsyncOpQueue.h b/Engine/Plugins/Online/OnlineServices/Source/OnlineServicesCommon/Public/Online/OnlineAsyncOpQueue.h
--- a/Engine/Plugins/Online/OnlineServices/Source/OnlineServicesCommon/Public/Online/OnlineAsyncOpQueue.h
+++ b/Engine/Plugins/Online/OnlineServices/Source/OnlineServicesCommon/Public/Online/OnlineAsyncOpQueue.h
@@ -2,10 +2,11 @@
 
 #pragma once
 
-#include "Online/OnlineAsyncOp.h"
 
-#include "Containers/Queue.h"
 #include "Containers/SpscQueue.h"
+#include "Online/OnlineAsyncOpHandle.h"
+
+namespace UE::Online { template <typename OpType> class TOnlineAsyncOp; }
 
 namespace UE::Online {
 
@@ -155,3 +156,8 @@ protected:
 };
 
 /* UE::Online */ }
+
+#if UE_ENABLE_INCLUDE_ORDER_DEPRECATED_IN_5_2
+#include "Containers/Queue.h"
+#include "Online/OnlineAsyncOp.h"
+#endif
diff --git a/Engine/Plugins/Online/OnlineServices/Source/OnlineServicesCommon/Public/Online/OnlineComponent.h b/Engine/Plugins/Online/OnlineServices/Source/OnlineServicesCommon/Public/Online/OnlineComponent.h
--- a/Engine/Plugins/Online/OnlineServices/Source/OnlineServicesCommon/Public/Online/OnlineComponent.h
+++ b/Engine/Plugins/Online/OnlineServices/Source/OnlineServicesCommon/Public/Online/OnlineComponent.h
@@ -2,9 +2,8 @@
 
 #pragma once
 
-#include "Templates/SharedPointer.h"
-#include "Online/OnlineAsyncOp.h"
 #include "Online/IOnlineComponent.h"
+#include "Online/OnlineExecHandler.h"
 #include "Online/OnlineServicesCommon.h"
 
 namespace UE::Online {
diff --git a/Engine/Plugins/Online/OnlineServices/Source/OnlineServicesCommon/Public/Online/OnlineComponentRegistry.h b/Engine/Plugins/Online/OnlineServices/Source/OnlineServicesCommon/Public/Online/OnlineComponentRegistry.h
--- a/Engine/Plugins/Online/OnlineServices/Source/OnlineServicesCommon/Public/Online/OnlineComponentRegistry.h
+++ b/Engine/Plugins/Online/OnlineServices/Source/OnlineServicesCommon/Public/Online/OnlineComponentRegistry.h
@@ -2,13 +2,12 @@
 
 #pragma once
 
-#include <Containers/Array.h>
-#include <Templates/UniquePtr.h>
 
-#include "Online/IOnlineComponent.h"
 #include "Online/OnlineTypeInfo.h"
 #include "Online/OnlineMeta.h"
 
+namespace UE::Online { class IOnlineComponent; }
+
 namespace UE::Online {
 
 class FOnlineComponentRegistry
@@ -120,3 +119,7 @@ private:
 };
 
 /* UE::Online */ }
+
+#if UE_ENABLE_INCLUDE_ORDER_DEPRECATED_IN_5_2
+#include "Online/IOnlineComponent.h"
+#endif
diff --git a/Engine/Plugins/Online/OnlineServices/Source/OnlineServicesCommon/Public/Online/OnlineConfig.h b/Engine/Plugins/Online/OnlineServices/Source/OnlineServicesCommon/Public/Online/OnlineConfig.h
--- a/Engine/Plugins/Online/OnlineServices/Source/OnlineServicesCommon/Public/Online/OnlineConfig.h
+++ b/Engine/Plugins/Online/OnlineServices/Source/OnlineServicesCommon/Public/Online/OnlineConfig.h
@@ -2,10 +2,10 @@
 
 #pragma once
 
-#include "Misc/ConfigCacheIni.h"
+#include "Internationalization/Text.h"
 #include "Online/OnlineMeta.h"
-#include "Online/SchemaTypes.h"
-#include "Algo/Transform.h"
+
+namespace UE::Online { class FSchemaVariant; }
 
 namespace UE::Online {
 
@@ -453,3 +453,9 @@ auto LoadConfig(IOnlineConfigProvider& Provider, const TArray<FString>& SectionH
 }
 
 /* UE::Online */ }
+
+#if UE_ENABLE_INCLUDE_ORDER_DEPRECATED_IN_5_2
+#include "Algo/Transform.h"
+#include "Misc/ConfigCacheIni.h"
+#include "Online/SchemaTypes.h"
+#endif
diff --git a/Engine/Plugins/Online/OnlineServices/Source/OnlineServicesCommon/Public/Online/OnlineExecHandler.h b/Engine/Plugins/Online/OnlineServices/Source/OnlineServicesCommon/Public/Online/OnlineExecHandler.h
--- a/Engine/Plugins/Online/OnlineServices/Source/OnlineServicesCommon/Public/Online/OnlineExecHandler.h
+++ b/Engine/Plugins/Online/OnlineServices/Source/OnlineServicesCommon/Public/Online/OnlineExecHandler.h
@@ -2,18 +2,16 @@
 
 #pragma once
 
-#include "Online/CoreOnline.h"
-#include "Online/OnlineUtils.h"
+#include "Online/OnlineAsyncOpHandle.h"
 #include "Online/OnlineServicesLog.h"
 #include "Online/OnlineResult.h"
-#include "Online/OnlineAsyncOp.h"
 
-#include "Online/OnlineServicesRegistry.h"
 #include "Online/Auth.h"
 
-#include "Templates/Invoke.h"
-#include "Templates/Models.h"
-#include "Templates/Tuple.h"
+#include "Online/OnlineServices.h"
+#include "Trace/Trace.inl"
+
+template <typename OptionalType> struct TOptional;
 
 /*
 	This file facilitates online command execution.
@@ -969,3 +967,9 @@ private:
 };
 
 /* UE::Online */ }
+
+#if UE_ENABLE_INCLUDE_ORDER_DEPRECATED_IN_5_2
+#include "Online/OnlineAsyncOp.h"
+#include "Online/OnlineServicesRegistry.h"
+#include "Online/OnlineUtils.h"
+#endif
diff --git a/Engine/Plugins/Online/OnlineServices/Source/OnlineServicesCommon/Public/Online/OnlineServicesCommon.h b/Engine/Plugins/Online/OnlineServices/Source/OnlineServicesCommon/Public/Online/OnlineServicesCommon.h
--- a/Engine/Plugins/Online/OnlineServices/Source/OnlineServicesCommon/Public/Online/OnlineServicesCommon.h
+++ b/Engine/Plugins/Online/OnlineServices/Source/OnlineServicesCommon/Public/Online/OnlineServicesCommon.h
@@ -8,13 +8,10 @@
 #include "Online/OnlineAsyncOpCache.h"
 #include "Online/OnlineAsyncOpQueue.h"
 #include "Online/OnlineConfig.h"
-#include "Online/OnlineExecHandler.h"
-#include "Online/OnlineServicesLog.h"
 
-#include "Async/Async.h"
 #include "Containers/Ticker.h"
-#include "Misc/CoreMisc.h"
-#include "Templates/SharedPointer.h"
+
+namespace UE::Online { class IOnlineExecHandler; }
 
 namespace UE::Online {
 
@@ -340,3 +337,8 @@ protected:
 
 /* UE::Online */ }
 
+
+#if UE_ENABLE_INCLUDE_ORDER_DEPRECATED_IN_5_2
+#include "Online/OnlineExecHandler.h"
+#include "Online/OnlineServicesLog.h"
+#endif
diff --git a/Engine/Plugins/Online/OnlineServices/Source/OnlineServicesCommon/Public/Online/OnlineTypeInfo.h b/Engine/Plugins/Online/OnlineServices/Source/OnlineServicesCommon/Public/Online/OnlineTypeInfo.h
--- a/Engine/Plugins/Online/OnlineServices/Source/OnlineServicesCommon/Public/Online/OnlineTypeInfo.h
+++ b/Engine/Plugins/Online/OnlineServices/Source/OnlineServicesCommon/Public/Online/OnlineTypeInfo.h
@@ -2,11 +2,12 @@
 
 #pragma once
 
-#include "Misc/CString.h"
 #include "Misc/GeneratedTypeName.h"
-#include "Templates/RemoveReference.h"
 
-#include "Online/OnlineMeta.h"
+#include "Templates/Models.h"
+#include "Templates/TypeHash.h"
+
+namespace UE::Online::Meta { struct CSuperDefined; }
 
 namespace UE::Online {
 
@@ -55,3 +56,8 @@ struct TOnlineTypeInfo
 };
 
 /* UE::Online */ }
+
+#if UE_ENABLE_INCLUDE_ORDER_DEPRECATED_IN_5_2
+#include "Online/OnlineMeta.h"
+#include "Templates/RemoveReference.h"
+#endif
diff --git a/Engine/Plugins/Online/OnlineServices/Source/OnlineServicesCommonEngineUtils/Public/Online/OnlineServicesCommonEngineUtils.h b/Engine/Plugins/Online/OnlineServices/Source/OnlineServicesCommonEngineUtils/Public/Online/OnlineServicesCommonEngineUtils.h
--- a/Engine/Plugins/Online/OnlineServices/Source/OnlineServicesCommonEngineUtils/Public/Online/OnlineServicesCommonEngineUtils.h
+++ b/Engine/Plugins/Online/OnlineServices/Source/OnlineServicesCommonEngineUtils/Public/Online/OnlineServicesCommonEngineUtils.h
@@ -2,11 +2,17 @@
 
 #pragma once
 
-#include "CoreMinimal.h"
-#include "Online/CoreOnline.h"
+#include "HAL/Platform.h"
+
+class FName;
 
 namespace UE::Online {
 
 ONLINESERVICESCOMMONENGINEUTILS_API int32 GetPortFromNetDriver(FName InstanceName);
 
-/* UE::Online */ }
\ No newline at end of file
+/* UE::Online */ }
+
+#if UE_ENABLE_INCLUDE_ORDER_DEPRECATED_IN_5_2
+#include "CoreMinimal.h"
+#include "Online/CoreOnline.h"
+#endif
diff --git a/Engine/Plugins/Online/OnlineServices/Source/OnlineServicesCommonEngineUtils/Public/Online/SessionsLAN.h b/Engine/Plugins/Online/OnlineServices/Source/OnlineServicesCommonEngineUtils/Public/Online/SessionsLAN.h
--- a/Engine/Plugins/Online/OnlineServices/Source/OnlineServicesCommonEngineUtils/Public/Online/SessionsLAN.h
+++ b/Engine/Plugins/Online/OnlineServices/Source/OnlineServicesCommonEngineUtils/Public/Online/SessionsLAN.h
@@ -1,11 +1,13 @@
 // Copyright Epic Games, Inc. All Rights Reserved.
 #pragma once
 
-#include "CoreMinimal.h"
-#include "Online/LANBeacon.h"
-#include "Online/NboSerializer.h"
 #include "Online/SessionsCommon.h"
 
+class FInternetAddr;
+class FLANSession;
+class FNboSerializeFromBuffer;
+class FNboSerializeToBuffer;
+
 namespace UE::Online {
 
 class FOnlineServicesCommon;
@@ -80,4 +82,10 @@ protected:
 	TSharedRef<FLANSession> LANSessionManager;
 };
 
-/* UE::Online */ }
\ No newline at end of file
+/* UE::Online */ }
+
+#if UE_ENABLE_INCLUDE_ORDER_DEPRECATED_IN_5_2
+#include "CoreMinimal.h"
+#include "Online/LANBeacon.h"
+#include "Online/NboSerializer.h"
+#endif
diff --git a/Engine/Plugins/Online/OnlineServices/Source/OnlineServicesInterface/Public/Online/Achievements.h b/Engine/Plugins/Online/OnlineServices/Source/OnlineServicesInterface/Public/Online/Achievements.h
--- a/Engine/Plugins/Online/OnlineServices/Source/OnlineServicesInterface/Public/Online/Achievements.h
+++ b/Engine/Plugins/Online/OnlineServices/Source/OnlineServicesInterface/Public/Online/Achievements.h
@@ -3,10 +3,14 @@
 #pragma once
 
 #include "Online/CoreOnline.h"
-#include "Online/OnlineAsyncOpHandle.h"
 #include "Online/OnlineMeta.h"
 #include "Misc/DateTime.h"
 
+namespace UE::Online { template <typename DelegateSignature> class TOnlineEvent; }
+namespace UE::Online { template <typename OpType> class TOnlineAsyncOpHandle; }
+namespace UE::Online { template <typename OpType> class TOnlineResult; }
+namespace UE::Online::Meta { template <typename StructType> struct TStructDetails; }
+
 namespace UE::Online {
 
 struct FQueryAchievementDefinitions
@@ -326,3 +330,7 @@ END_ONLINE_STRUCT_META()
 /* Meta*/ }
 
 /* UE::Online */ }
+
+#if UE_ENABLE_INCLUDE_ORDER_DEPRECATED_IN_5_2
+#include "Online/OnlineAsyncOpHandle.h"
+#endif
diff --git a/Engine/Plugins/Online/OnlineServices/Source/OnlineServicesInterface/Public/Online/AchievementsErrors.h b/Engine/Plugins/Online/OnlineServices/Source/OnlineServicesInterface/Public/Online/AchievementsErrors.h
--- a/Engine/Plugins/Online/OnlineServices/Source/OnlineServicesInterface/Public/Online/AchievementsErrors.h
+++ b/Engine/Plugins/Online/OnlineServices/Source/OnlineServicesInterface/Public/Online/AchievementsErrors.h
@@ -3,6 +3,7 @@
 #pragma once
 
 #include "Online/OnlineError.h"
+#include "Online/OnlineErrorCode.h"
 
 #define LOCTEXT_NAMESPACE "AchievementsErrors"
 
diff --git a/Engine/Plugins/Online/OnlineServices/Source/OnlineServicesInterface/Public/Online/Auth.h b/Engine/Plugins/Online/OnlineServices/Source/OnlineServicesInterface/Public/Online/Auth.h
--- a/Engine/Plugins/Online/OnlineServices/Source/OnlineServicesInterface/Public/Online/Auth.h
+++ b/Engine/Plugins/Online/OnlineServices/Source/OnlineServicesInterface/Public/Online/Auth.h
@@ -2,13 +2,13 @@
 
 #pragma once
 
-#include "Containers/Map.h"
-#include "Containers/Set.h"
 
 #include "Online/CoreOnline.h"
-#include "Online/OnlineAsyncOpHandle.h"
-#include "Online/OnlineMeta.h"
-#include "Online/Schema.h"
+#include "Online/SchemaTypes.h"
+
+namespace UE::Online { template <typename DelegateSignature> class TOnlineEvent; }
+namespace UE::Online { template <typename OpType> class TOnlineAsyncOpHandle; }
+namespace UE::Online { template <typename OpType> class TOnlineResult; }
 
 namespace UE::Online {
 
@@ -652,3 +652,8 @@ END_ONLINE_STRUCT_META()
 /* Meta*/ }
 
 /* UE::Online */ }
+
+#if UE_ENABLE_INCLUDE_ORDER_DEPRECATED_IN_5_2
+#include "Online/OnlineAsyncOpHandle.h"
+#include "Online/Schema.h"
+#endif
diff --git a/Engine/Plugins/Online/OnlineServices/Source/OnlineServicesInterface/Public/Online/AuthErrors.h b/Engine/Plugins/Online/OnlineServices/Source/OnlineServicesInterface/Public/Online/AuthErrors.h
--- a/Engine/Plugins/Online/OnlineServices/Source/OnlineServicesInterface/Public/Online/AuthErrors.h
+++ b/Engine/Plugins/Online/OnlineServices/Source/OnlineServicesInterface/Public/Online/AuthErrors.h
@@ -3,6 +3,7 @@
 #pragma once
 
 #include "Online/OnlineError.h"
+#include "Online/OnlineErrorCode.h"
 
 #define LOCTEXT_NAMESPACE "AuthErrors"
 
diff --git a/Engine/Plugins/Online/OnlineServices/Source/OnlineServicesInterface/Public/Online/Commerce.h b/Engine/Plugins/Online/OnlineServices/Source/OnlineServicesInterface/Public/Online/Commerce.h
--- a/Engine/Plugins/Online/OnlineServices/Source/OnlineServicesInterface/Public/Online/Commerce.h
+++ b/Engine/Plugins/Online/OnlineServices/Source/OnlineServicesInterface/Public/Online/Commerce.h
@@ -2,10 +2,15 @@
 
 #pragma once
 
+#include "Misc/DateTime.h"
 #include "Online/CoreOnline.h"
-#include "Online/OnlineAsyncOpHandle.h"
 #include "Online/OnlineMeta.h"
 
+namespace UE::Online { template <typename DelegateSignature> class TOnlineEvent; }
+namespace UE::Online { template <typename OpType> class TOnlineAsyncOpHandle; }
+namespace UE::Online { template <typename OpType> class TOnlineResult; }
+namespace UE::Online::Meta { template <typename StructType> struct TStructDetails; }
+
 namespace UE::Online {
 
 /* Unique identifier for an offer */
@@ -476,3 +481,7 @@ END_ONLINE_STRUCT_META()
 
 
 
+
+#if UE_ENABLE_INCLUDE_ORDER_DEPRECATED_IN_5_2
+#include "Online/OnlineAsyncOpHandle.h"
+#endif
diff --git a/Engine/Plugins/Online/OnlineServices/Source/OnlineServicesInterface/Public/Online/Connectivity.h b/Engine/Plugins/Online/OnlineServices/Source/OnlineServicesInterface/Public/Online/Connectivity.h
--- a/Engine/Plugins/Online/OnlineServices/Source/OnlineServicesInterface/Public/Online/Connectivity.h
+++ b/Engine/Plugins/Online/OnlineServices/Source/OnlineServicesInterface/Public/Online/Connectivity.h
@@ -2,9 +2,14 @@
 
 #pragma once
 
-#include "Online/OnlineAsyncOpHandle.h"
+#include "Containers/UnrealString.h"
 #include "Online/OnlineMeta.h"
 
+namespace UE::Online { template <typename DelegateSignature> class TOnlineEvent; }
+namespace UE::Online { template <typename OpType> class TOnlineResult; }
+namespace UE::Online::Meta { template <typename StructType> struct TStructDetails; }
+struct FScriptContainerElement;
+
 namespace UE::Online {
 
 enum class EOnlineServicesConnectionStatus : uint8
@@ -80,4 +85,8 @@ END_ONLINE_STRUCT_META()
 
 /* Meta*/}
 
-/* UE::Online */}
\ No newline at end of file
+/* UE::Online */}
+
+#if UE_ENABLE_INCLUDE_ORDER_DEPRECATED_IN_5_2
+#include "Online/OnlineAsyncOpHandle.h"
+#endif
diff --git a/Engine/Plugins/Online/OnlineServices/Source/OnlineServicesInterface/Public/Online/ExternalUI.h b/Engine/Plugins/Online/OnlineServices/Source/OnlineServicesInterface/Public/Online/ExternalUI.h
--- a/Engine/Plugins/Online/OnlineServices/Source/OnlineServicesInterface/Public/Online/ExternalUI.h
+++ b/Engine/Plugins/Online/OnlineServices/Source/OnlineServicesInterface/Public/Online/ExternalUI.h
@@ -2,11 +2,13 @@
 
 #pragma once
 
-#include "Online/Auth.h"
 #include "Online/CoreOnline.h"
-#include "Online/OnlineAsyncOpHandle.h"
 #include "Online/OnlineMeta.h"
 
+namespace UE::Online { template <typename DelegateSignature> class TOnlineEvent; }
+namespace UE::Online { template <typename OpType> class TOnlineAsyncOpHandle; }
+namespace UE::Online::Meta { template <typename StructType> struct TStructDetails; }
+
 namespace UE::Online {
 
 class FOnlineError;
@@ -98,3 +100,8 @@ END_ONLINE_STRUCT_META()
 /* Meta*/ }
 
 /* UE::Online */ }
+
+#if UE_ENABLE_INCLUDE_ORDER_DEPRECATED_IN_5_2
+#include "Online/Auth.h"
+#include "Online/OnlineAsyncOpHandle.h"
+#endif
diff --git a/Engine/Plugins/Online/OnlineServices/Source/OnlineServicesInterface/Public/Online/Leaderboards.h b/Engine/Plugins/Online/OnlineServices/Source/OnlineServicesInterface/Public/Online/Leaderboards.h
--- a/Engine/Plugins/Online/OnlineServices/Source/OnlineServicesInterface/Public/Online/Leaderboards.h
+++ b/Engine/Plugins/Online/OnlineServices/Source/OnlineServicesInterface/Public/Online/Leaderboards.h
@@ -3,9 +3,11 @@
 #pragma once
 
 #include "Online/CoreOnline.h"
-#include "Online/OnlineAsyncOpHandle.h"
 #include "Online/OnlineMeta.h"
 
+namespace UE::Online { template <typename OpType> class TOnlineAsyncOpHandle; }
+namespace UE::Online::Meta { template <typename StructType> struct TStructDetails; }
+
 #define UE_LEADERBOARD_RANK_UNKNOWN -1
 
 namespace UE::Online {
@@ -160,3 +162,7 @@ END_ONLINE_STRUCT_META()
 /* Meta*/ }
 
 /* UE::Online */ }
+
+#if UE_ENABLE_INCLUDE_ORDER_DEPRECATED_IN_5_2
+#include "Online/OnlineAsyncOpHandle.h"
+#endif
diff --git a/Engine/Plugins/Online/OnlineServices/Source/OnlineServicesInterface/Public/Online/Lobbies.h b/Engine/Plugins/Online/OnlineServices/Source/OnlineServicesInterface/Public/Online/Lobbies.h
--- a/Engine/Plugins/Online/OnlineServices/Source/OnlineServicesInterface/Public/Online/Lobbies.h
+++ b/Engine/Plugins/Online/OnlineServices/Source/OnlineServicesInterface/Public/Online/Lobbies.h
@@ -2,11 +2,13 @@
 
 #pragma once
 
-#include "Online/OnlineAsyncOpHandle.h"
 #include "Online/CoreOnline.h"
-#include "Online/OnlineMeta.h"
-#include "Online/Schema.h"
-#include "Misc/TVariant.h"
+#include "Online/OnlineError.h"
+#include "Online/OnlineResult.h"
+#include "Online/SchemaTypes.h"
+
+namespace UE::Online { template <typename DelegateSignature> class TOnlineEvent; }
+namespace UE::Online { template <typename OpType> class TOnlineAsyncOpHandle; }
 
 namespace UE::Online {
 
@@ -1026,3 +1028,8 @@ END_ONLINE_STRUCT_META()
 /* Meta*/ }
 
 /* UE::Online */ }
+
+#if UE_ENABLE_INCLUDE_ORDER_DEPRECATED_IN_5_2
+#include "Online/OnlineAsyncOpHandle.h"
+#include "Online/Schema.h"
+#endif
diff --git a/Engine/Plugins/Online/OnlineServices/Source/OnlineServicesInterface/Public/Online/OnlineAsyncOpHandle.h b/Engine/Plugins/Online/OnlineServices/Source/OnlineServicesInterface/Public/Online/OnlineAsyncOpHandle.h
--- a/Engine/Plugins/Online/OnlineServices/Source/OnlineServicesInterface/Public/Online/OnlineAsyncOpHandle.h
+++ b/Engine/Plugins/Online/OnlineServices/Source/OnlineServicesInterface/Public/Online/OnlineAsyncOpHandle.h
@@ -2,16 +2,11 @@
 
 #pragma once
 
-#include "CoreMinimal.h"
-#include "Online/OnlineError.h"
-#include "Online/OnlineResult.h"
 #include "Online/OnlineErrorDefinitions.h"
-#include "Templates/Models.h"
 
-#include "Templates/SharedPointer.h"
-#include "Delegates/Delegate.h"
 
-#include <type_traits>
+
+namespace UE::Online { template <typename OpType> class TOnlineResult; }
 
 class UObjectBase;
 
@@ -500,3 +495,8 @@ private:
 };
 
 /* UE::Online */ }
+
+#if UE_ENABLE_INCLUDE_ORDER_DEPRECATED_IN_5_2
+#include "CoreMinimal.h"
+#include "Online/OnlineResult.h"
+#endif
diff --git a/Engine/Plugins/Online/OnlineServices/Source/OnlineServicesInterface/Public/Online/OnlineError.h b/Engine/Plugins/Online/OnlineServices/Source/OnlineServicesInterface/Public/Online/OnlineError.h
--- a/Engine/Plugins/Online/OnlineServices/Source/OnlineServicesInterface/Public/Online/OnlineError.h
+++ b/Engine/Plugins/Online/OnlineServices/Source/OnlineServicesInterface/Public/Online/OnlineError.h
@@ -1,9 +1,8 @@
 // Copyright Epic Games, Inc. All Rights Reserved.
 #pragma once
 
-#include "Templates/SharedPointer.h"
+#include "Internationalization/Internationalization.h"
 #include "OnlineErrorCode.h"
-#include "CoreMinimal.h"
 
 namespace UE::Online {
 
@@ -186,4 +185,8 @@ inline FString ToLogString(const FOnlineError& Error)
 
 } // namespace UE::Online
 
-#undef LOCTEXT_NAMESPACE
\ No newline at end of file
+#undef LOCTEXT_NAMESPACE
+
+#if UE_ENABLE_INCLUDE_ORDER_DEPRECATED_IN_5_2
+#include "CoreMinimal.h"
+#endif
diff --git a/Engine/Plugins/Online/OnlineServices/Source/OnlineServicesInterface/Public/Online/OnlineErrorCode.h b/Engine/Plugins/Online/OnlineServices/Source/OnlineServicesInterface/Public/Online/OnlineErrorCode.h
--- a/Engine/Plugins/Online/OnlineServices/Source/OnlineServicesInterface/Public/Online/OnlineErrorCode.h
+++ b/Engine/Plugins/Online/OnlineServices/Source/OnlineServicesInterface/Public/Online/OnlineErrorCode.h
@@ -1,9 +1,10 @@
 // Copyright Epic Games, Inc. All Rights Reserved.
 #pragma once
 
-#include "Containers/UnrealString.h"
 #include "CoreTypes.h"
 
+class FString;
+
 namespace UE::Online::Errors {
 using ErrorCodeType = uint64;
 
@@ -92,4 +93,8 @@ namespace ErrorCode { namespace Category { \
     namespace ErrorCode { namespace CategoryName { static constexpr ErrorCodeType Name = Create(Category::CategoryName, ErrorCodeValue); } } \
 	UE_ONLINE_ERROR_INTERNAL(CategoryName, Name, ErrorCodeValue, ErrorMessage, ErrorText)
 
-} /* namespace UE::Online::Errors */
\ No newline at end of file
+} /* namespace UE::Online::Errors */
+
+#if UE_ENABLE_INCLUDE_ORDER_DEPRECATED_IN_5_2
+#include "Containers/UnrealString.h"
+#endif
diff --git a/Engine/Plugins/Online/OnlineServices/Source/OnlineServicesInterface/Public/Online/OnlineErrorDefinitions.h b/Engine/Plugins/Online/OnlineServices/Source/OnlineServicesInterface/Public/Online/OnlineErrorDefinitions.h
--- a/Engine/Plugins/Online/OnlineServices/Source/OnlineServicesInterface/Public/Online/OnlineErrorDefinitions.h
+++ b/Engine/Plugins/Online/OnlineServices/Source/OnlineServicesInterface/Public/Online/OnlineErrorDefinitions.h
@@ -2,6 +2,7 @@
 #pragma once
 
 #include "Online/OnlineError.h"
+#include "Online/OnlineErrorCode.h"
 
 #define LOCTEXT_NAMESPACE "OnlineErrors"
 
@@ -35,4 +36,4 @@ UE_ONLINE_ERROR_COMMON(Common, Unknown, 22, TEXT("unknown"), LOCTEXT("Unknown",
 
 } /* namespace UE::Online::Errors */
 
-#undef LOCTEXT_NAMESPACE
\ No newline at end of file
+#undef LOCTEXT_NAMESPACE
diff --git a/Engine/Plugins/Online/OnlineServices/Source/OnlineServicesInterface/Public/Online/OnlineMeta.h b/Engine/Plugins/Online/OnlineServices/Source/OnlineServicesInterface/Public/Online/OnlineMeta.h
--- a/Engine/Plugins/Online/OnlineServices/Source/OnlineServicesInterface/Public/Online/OnlineMeta.h
+++ b/Engine/Plugins/Online/OnlineServices/Source/OnlineServicesInterface/Public/Online/OnlineMeta.h
@@ -2,7 +2,6 @@
 
 #pragma once
 
-#include "Templates/Models.h"
 #include "Templates/Tuple.h"
 
 namespace UE::Online::Meta {
diff --git a/Engine/Plugins/Online/OnlineServices/Source/OnlineServicesInterface/Public/Online/OnlineResult.h b/Engine/Plugins/Online/OnlineServices/Source/OnlineServicesInterface/Public/Online/OnlineResult.h
--- a/Engine/Plugins/Online/OnlineServices/Source/OnlineServicesInterface/Public/Online/OnlineResult.h
+++ b/Engine/Plugins/Online/OnlineServices/Source/OnlineServicesInterface/Public/Online/OnlineResult.h
@@ -2,9 +2,11 @@
 
 #pragma once
 
+#include "CoreGlobals.h"
 #include "Misc/TVariant.h"
 
-#include "Online/OnlineError.h"
+
+namespace UE::Online { class FOnlineError; }
 
 namespace UE::Online {
 
@@ -218,3 +220,7 @@ FString ToLogString(const TOnlineResult<T>& Result)
 }
 
 /* UE::Online */ }
+
+#if UE_ENABLE_INCLUDE_ORDER_DEPRECATED_IN_5_2
+#include "Online/OnlineError.h"
+#endif
diff --git a/Engine/Plugins/Online/OnlineServices/Source/OnlineServicesInterface/Public/Online/OnlineServices.h b/Engine/Plugins/Online/OnlineServices/Source/OnlineServicesInterface/Public/Online/OnlineServices.h
--- a/Engine/Plugins/Online/OnlineServices/Source/OnlineServicesInterface/Public/Online/OnlineServices.h
+++ b/Engine/Plugins/Online/OnlineServices/Source/OnlineServicesInterface/Public/Online/OnlineServices.h
@@ -3,10 +3,11 @@
 #pragma once
 
 #include "Online/CoreOnline.h"
-#include "Templates/SharedPointer.h"
-#include "Online/OnlineAsyncOpHandle.h"
 #include "Online/OnlineMeta.h"
 
+namespace UE::Online { template <typename OpType> class TOnlineResult; }
+namespace UE::Online::Meta { template <typename StructType> struct TStructDetails; }
+
 class FString;
 
 namespace UE::Online {
@@ -228,3 +229,7 @@ END_ONLINE_STRUCT_META()
 /* Meta*/}
 
 /* UE::Online */ }
+
+#if UE_ENABLE_INCLUDE_ORDER_DEPRECATED_IN_5_2
+#include "Online/OnlineAsyncOpHandle.h"
+#endif
diff --git a/Engine/Plugins/Online/OnlineServices/Source/OnlineServicesInterface/Public/Online/OnlineServicesDelegates.h b/Engine/Plugins/Online/OnlineServices/Source/OnlineServicesInterface/Public/Online/OnlineServicesDelegates.h
--- a/Engine/Plugins/Online/OnlineServices/Source/OnlineServicesInterface/Public/Online/OnlineServicesDelegates.h
+++ b/Engine/Plugins/Online/OnlineServices/Source/OnlineServicesInterface/Public/Online/OnlineServicesDelegates.h
@@ -2,7 +2,7 @@
 
 #pragma once
 
-#include "CoreMinimal.h"
+#include "Delegates/Delegate.h"
 
 namespace UE::Online {
 
@@ -19,3 +19,7 @@ DECLARE_MULTICAST_DELEGATE_OneParam(FOnOnlineServicesCreated, TSharedRef<class I
 extern ONLINESERVICESINTERFACE_API FOnOnlineServicesCreated OnOnlineServicesCreated;
 
 /* UE::Online */ }
+
+#if UE_ENABLE_INCLUDE_ORDER_DEPRECATED_IN_5_2
+#include "CoreMinimal.h"
+#endif
diff --git a/Engine/Plugins/Online/OnlineServices/Source/OnlineServicesInterface/Public/Online/OnlineServicesRegistry.h b/Engine/Plugins/Online/OnlineServices/Source/OnlineServicesInterface/Public/Online/OnlineServicesRegistry.h
--- a/Engine/Plugins/Online/OnlineServices/Source/OnlineServicesInterface/Public/Online/OnlineServicesRegistry.h
+++ b/Engine/Plugins/Online/OnlineServices/Source/OnlineServicesInterface/Public/Online/OnlineServicesRegistry.h
@@ -2,11 +2,10 @@
 
 #pragma once
 
-#include "Containers/Map.h"
-#include "Templates/SharedPointer.h"
-#include "Templates/UniquePtr.h"
 
-#include "Online/OnlineServices.h"
+#include "Online/CoreOnline.h"
+
+namespace UE::Online { class IOnlineServices; }
 
 class FLazySingleton;
 
@@ -149,3 +148,7 @@ private:
 };
 
 /* UE::Online */ }
+
+#if UE_ENABLE_INCLUDE_ORDER_DEPRECATED_IN_5_2
+#include "Online/OnlineServices.h"
+#endif
diff --git a/Engine/Plugins/Online/OnlineServices/Source/OnlineServicesInterface/Public/Online/Presence.h b/Engine/Plugins/Online/OnlineServices/Source/OnlineServicesInterface/Public/Online/Presence.h
--- a/Engine/Plugins/Online/OnlineServices/Source/OnlineServicesInterface/Public/Online/Presence.h
+++ b/Engine/Plugins/Online/OnlineServices/Source/OnlineServicesInterface/Public/Online/Presence.h
@@ -3,10 +3,15 @@
 #pragma once
 
 #include "Online/CoreOnline.h"
-#include "Online/OnlineAsyncOpHandle.h"
+#include "Online/OnlineErrorCode.h"
 #include "Online/OnlineMeta.h"
 #include "Online/OnlineError.h"
 
+namespace UE::Online { template <typename DelegateSignature> class TOnlineEvent; }
+namespace UE::Online { template <typename OpType> class TOnlineAsyncOpHandle; }
+namespace UE::Online { template <typename OpType> class TOnlineResult; }
+namespace UE::Online::Meta { template <typename StructType> struct TStructDetails; }
+
 #define LOCTEXT_NAMESPACE "Presence"
 
 namespace UE::Online {
@@ -425,4 +430,8 @@ END_ONLINE_STRUCT_META()
 
 /* UE::Online */ }
 
-#undef LOCTEXT_NAMESPACE
\ No newline at end of file
+#undef LOCTEXT_NAMESPACE
+
+#if UE_ENABLE_INCLUDE_ORDER_DEPRECATED_IN_5_2
+#include "Online/OnlineAsyncOpHandle.h"
+#endif
diff --git a/Engine/Plugins/Online/OnlineServices/Source/OnlineServicesInterface/Public/Online/Privileges.h b/Engine/Plugins/Online/OnlineServices/Source/OnlineServicesInterface/Public/Online/Privileges.h
--- a/Engine/Plugins/Online/OnlineServices/Source/OnlineServicesInterface/Public/Online/Privileges.h
+++ b/Engine/Plugins/Online/OnlineServices/Source/OnlineServicesInterface/Public/Online/Privileges.h
@@ -3,9 +3,11 @@
 #pragma once
 
 #include "Online/CoreOnline.h"
-#include "Online/OnlineAsyncOpHandle.h"
 #include "Online/OnlineMeta.h"
 
+namespace UE::Online { template <typename OpType> class TOnlineAsyncOpHandle; }
+namespace UE::Online::Meta { template <typename StructType> struct TStructDetails; }
+
 namespace UE::Online {
 
 enum class EUserPrivileges : uint8
@@ -101,4 +103,8 @@ END_ONLINE_STRUCT_META()
 
 /* Meta*/}
 
-/* UE::Online */}
\ No newline at end of file
+/* UE::Online */}
+
+#if UE_ENABLE_INCLUDE_ORDER_DEPRECATED_IN_5_2
+#include "Online/OnlineAsyncOpHandle.h"
+#endif
diff --git a/Engine/Plugins/Online/OnlineServices/Source/OnlineServicesInterface/Public/Online/Schema.h b/Engine/Plugins/Online/OnlineServices/Source/OnlineServicesInterface/Public/Online/Schema.h
--- a/Engine/Plugins/Online/OnlineServices/Source/OnlineServicesInterface/Public/Online/Schema.h
+++ b/Engine/Plugins/Online/OnlineServices/Source/OnlineServicesInterface/Public/Online/Schema.h
@@ -2,8 +2,10 @@
 
 #pragma once
 
-#include "CoreMinimal.h"
 #include "Online/SchemaTypes.h"
+#include "Templates/SharedPointer.h"
+
+namespace UE::Online { template <typename OpType> class TOnlineResult; }
 
 namespace UE::Online {
 
@@ -180,3 +182,7 @@ private:
 };
 
 /* UE::Online */ }
+
+#if UE_ENABLE_INCLUDE_ORDER_DEPRECATED_IN_5_2
+#include "CoreMinimal.h"
+#endif
diff --git a/Engine/Plugins/Online/OnlineServices/Source/OnlineServicesInterface/Public/Online/SchemaTypes.h b/Engine/Plugins/Online/OnlineServices/Source/OnlineServicesInterface/Public/Online/SchemaTypes.h
--- a/Engine/Plugins/Online/OnlineServices/Source/OnlineServicesInterface/Public/Online/SchemaTypes.h
+++ b/Engine/Plugins/Online/OnlineServices/Source/OnlineServicesInterface/Public/Online/SchemaTypes.h
@@ -2,10 +2,13 @@
 
 #pragma once
 
-#include "CoreMinimal.h"
+#include "Containers/Map.h"
 #include "Misc/TVariant.h"
+#include "Misc/TVariantMeta.h"
 #include "Online/OnlineMeta.h"
-#include "Online/OnlineResult.h"
+#include "UObject/NameTypes.h"
+
+namespace UE::Online::Meta { template <typename StructType> struct TStructDetails; }
 
 namespace UE::Online {
 
@@ -715,3 +718,8 @@ END_ONLINE_STRUCT_META()
 /* Meta*/ }
 
 /* UE::Online */ }
+
+#if UE_ENABLE_INCLUDE_ORDER_DEPRECATED_IN_5_2
+#include "CoreMinimal.h"
+#include "Online/OnlineResult.h"
+#endif
diff --git a/Engine/Plugins/Online/OnlineServices/Source/OnlineServicesInterface/Public/Online/Sessions.h b/Engine/Plugins/Online/OnlineServices/Source/OnlineServicesInterface/Public/Online/Sessions.h
--- a/Engine/Plugins/Online/OnlineServices/Source/OnlineServicesInterface/Public/Online/Sessions.h
+++ b/Engine/Plugins/Online/OnlineServices/Source/OnlineServicesInterface/Public/Online/Sessions.h
@@ -2,11 +2,13 @@
 
 #pragma once
 
-#include "Online/OnlineAsyncOpHandle.h"
 #include "Online/CoreOnline.h"
-#include "Online/OnlineMeta.h"
+#include "Online/OnlineError.h"
 #include "Online/SchemaTypes.h"
-#include "Misc/TVariant.h"
+#include "Online/OnlineResult.h"
+
+namespace UE::Online { template <typename DelegateSignature> class TOnlineEvent; }
+namespace UE::Online { template <typename OpType> class TOnlineAsyncOpHandle; }
 
 namespace UE::Online {
 
@@ -1102,4 +1104,8 @@ END_ONLINE_STRUCT_META()
 
 /* Meta*/ }
 
-/* UE::Online */ }
\ No newline at end of file
+/* UE::Online */ }
+
+#if UE_ENABLE_INCLUDE_ORDER_DEPRECATED_IN_5_2
+#include "Online/OnlineAsyncOpHandle.h"
+#endif
diff --git a/Engine/Plugins/Online/OnlineServices/Source/OnlineServicesInterface/Public/Online/Social.h b/Engine/Plugins/Online/OnlineServices/Source/OnlineServicesInterface/Public/Online/Social.h
--- a/Engine/Plugins/Online/OnlineServices/Source/OnlineServicesInterface/Public/Online/Social.h
+++ b/Engine/Plugins/Online/OnlineServices/Source/OnlineServicesInterface/Public/Online/Social.h
@@ -3,9 +3,13 @@
 #pragma once
 
 #include "Online/CoreOnline.h"
-#include "Online/OnlineAsyncOpHandle.h"
 #include "Online/OnlineMeta.h"
 
+namespace UE::Online { template <typename DelegateSignature> class TOnlineEvent; }
+namespace UE::Online { template <typename OpType> class TOnlineAsyncOpHandle; }
+namespace UE::Online { template <typename OpType> class TOnlineResult; }
+namespace UE::Online::Meta { template <typename StructType> struct TStructDetails; }
+
 namespace UE::Online {
 
 /**
@@ -385,3 +389,7 @@ END_ONLINE_STRUCT_META()
 /* Meta*/ }
 
 /* UE::Online */ }
+
+#if UE_ENABLE_INCLUDE_ORDER_DEPRECATED_IN_5_2
+#include "Online/OnlineAsyncOpHandle.h"
+#endif
diff --git a/Engine/Plugins/Online/OnlineServices/Source/OnlineServicesInterface/Public/Online/Stats.h b/Engine/Plugins/Online/OnlineServices/Source/OnlineServicesInterface/Public/Online/Stats.h
--- a/Engine/Plugins/Online/OnlineServices/Source/OnlineServicesInterface/Public/Online/Stats.h
+++ b/Engine/Plugins/Online/OnlineServices/Source/OnlineServicesInterface/Public/Online/Stats.h
@@ -3,9 +3,11 @@
 #pragma once
 
 #include "Online/CoreOnline.h"
-#include "Online/OnlineAsyncOpHandle.h"
-#include "Online/OnlineMeta.h"
-#include "Online/Schema.h"
+#include "Online/SchemaTypes.h"
+
+namespace UE::Online { template <typename DelegateSignature> class TOnlineEvent; }
+namespace UE::Online { template <typename OpType> class TOnlineAsyncOpHandle; }
+namespace UE::Online { template <typename OpType> class TOnlineResult; }
 
 namespace UE::Online {
 
@@ -214,3 +216,8 @@ END_ONLINE_STRUCT_META()
 /* Meta*/ }
 
 /* UE::Online */ }
+
+#if UE_ENABLE_INCLUDE_ORDER_DEPRECATED_IN_5_2
+#include "Online/OnlineAsyncOpHandle.h"
+#include "Online/Schema.h"
+#endif
diff --git a/Engine/Plugins/Online/OnlineServices/Source/OnlineServicesInterface/Public/Online/TitleFile.h b/Engine/Plugins/Online/OnlineServices/Source/OnlineServicesInterface/Public/Online/TitleFile.h
--- a/Engine/Plugins/Online/OnlineServices/Source/OnlineServicesInterface/Public/Online/TitleFile.h
+++ b/Engine/Plugins/Online/OnlineServices/Source/OnlineServicesInterface/Public/Online/TitleFile.h
@@ -3,9 +3,12 @@
 #pragma once
 
 #include "Online/CoreOnline.h"
-#include "Online/OnlineAsyncOpHandle.h"
 #include "Online/OnlineMeta.h"
 
+namespace UE::Online { template <typename OpType> class TOnlineAsyncOpHandle; }
+namespace UE::Online { template <typename OpType> class TOnlineResult; }
+namespace UE::Online::Meta { template <typename StructType> struct TStructDetails; }
+
 namespace UE::Online {
 
 struct FTitleFileEnumerateFiles
@@ -113,3 +116,7 @@ END_ONLINE_STRUCT_META()
 /* Meta*/ }
 
 /* UE::Online */ }
+
+#if UE_ENABLE_INCLUDE_ORDER_DEPRECATED_IN_5_2
+#include "Online/OnlineAsyncOpHandle.h"
+#endif
diff --git a/Engine/Plugins/Online/OnlineServices/Source/OnlineServicesInterface/Public/Online/UserFile.h b/Engine/Plugins/Online/OnlineServices/Source/OnlineServicesInterface/Public/Online/UserFile.h
--- a/Engine/Plugins/Online/OnlineServices/Source/OnlineServicesInterface/Public/Online/UserFile.h
+++ b/Engine/Plugins/Online/OnlineServices/Source/OnlineServicesInterface/Public/Online/UserFile.h
@@ -3,9 +3,12 @@
 #pragma once
 
 #include "Online/CoreOnline.h"
-#include "Online/OnlineAsyncOpHandle.h"
 #include "Online/OnlineMeta.h"
 
+namespace UE::Online { template <typename OpType> class TOnlineAsyncOpHandle; }
+namespace UE::Online { template <typename OpType> class TOnlineResult; }
+namespace UE::Online::Meta { template <typename StructType> struct TStructDetails; }
+
 namespace UE::Online {
 
 struct FUserFileEnumerateFiles
@@ -208,3 +211,7 @@ END_ONLINE_STRUCT_META()
 /* Meta*/ }
 
 /* UE::Online */ }
+
+#if UE_ENABLE_INCLUDE_ORDER_DEPRECATED_IN_5_2
+#include "Online/OnlineAsyncOpHandle.h"
+#endif
diff --git a/Engine/Plugins/Online/OnlineServices/Source/OnlineServicesInterface/Public/Online/UserInfo.h b/Engine/Plugins/Online/OnlineServices/Source/OnlineServicesInterface/Public/Online/UserInfo.h
--- a/Engine/Plugins/Online/OnlineServices/Source/OnlineServicesInterface/Public/Online/UserInfo.h
+++ b/Engine/Plugins/Online/OnlineServices/Source/OnlineServicesInterface/Public/Online/UserInfo.h
@@ -3,9 +3,12 @@
 #pragma once
 
 #include "Online/CoreOnline.h"
-#include "Online/OnlineAsyncOpHandle.h"
 #include "Online/OnlineMeta.h"
 
+namespace UE::Online { template <typename OpType> class TOnlineAsyncOpHandle; }
+namespace UE::Online { template <typename OpType> class TOnlineResult; }
+namespace UE::Online::Meta { template <typename StructType> struct TStructDetails; }
+
 namespace UE::Online {
 
 struct FUserInfo
@@ -230,4 +233,8 @@ END_ONLINE_STRUCT_META()
 
 /* Meta*/ }
 
-/* UE::Online */ }
\ No newline at end of file
+/* UE::Online */ }
+
+#if UE_ENABLE_INCLUDE_ORDER_DEPRECATED_IN_5_2
+#include "Online/OnlineAsyncOpHandle.h"
+#endif
diff --git a/Engine/Plugins/Online/OnlineServicesEOSGS/Source/Private/Online/OnlineServicesEOSGSPlatformFactory.cpp b/Engine/Plugins/Online/OnlineServicesEOSGS/Source/Private/Online/OnlineServicesEOSGSPlatformFactory.cpp
--- a/Engine/Plugins/Online/OnlineServicesEOSGS/Source/Private/Online/OnlineServicesEOSGSPlatformFactory.cpp
+++ b/Engine/Plugins/Online/OnlineServicesEOSGS/Source/Private/Online/OnlineServicesEOSGSPlatformFactory.cpp
@@ -3,6 +3,7 @@
 #include "Online/OnlineServicesEOSGSPlatformFactory.h"
 
 #include "EOSShared.h"
+#include "Online/OnlineExecHandler.h"
 #include "Online/OnlineServicesEOSGS.h"
 
 #include "Misc/ConfigCacheIni.h"
diff --git a/Engine/Plugins/Online/OnlineServicesOSSAdapter/Source/Private/Online/AuthOSSAdapter.cpp b/Engine/Plugins/Online/OnlineServicesOSSAdapter/Source/Private/Online/AuthOSSAdapter.cpp
--- a/Engine/Plugins/Online/OnlineServicesOSSAdapter/Source/Private/Online/AuthOSSAdapter.cpp
+++ b/Engine/Plugins/Online/OnlineServicesOSSAdapter/Source/Private/Online/AuthOSSAdapter.cpp
@@ -9,6 +9,7 @@
 #include "Online/OnlineErrorDefinitions.h"
 #include "Online/DelegateAdapter.h"
 
+#include "Online/OnlineUtils.h"
 #include "OnlineSubsystem.h"
 #include "Interfaces/OnlineIdentityInterface.h"
 
diff --git a/Engine/Plugins/Online/OnlineSubsystem/Source/Private/OnlineError.cpp b/Engine/Plugins/Online/OnlineSubsystem/Source/Private/OnlineError.cpp
--- a/Engine/Plugins/Online/OnlineSubsystem/Source/Private/OnlineError.cpp
+++ b/Engine/Plugins/Online/OnlineSubsystem/Source/Private/OnlineError.cpp
@@ -1,7 +1,6 @@
 // Copyright Epic Games, Inc. All Rights Reserved.
 
 #include "OnlineError.h"
-#include "Internationalization/Internationalization.h"
 
 #define LOCTEXT_NAMESPACE "OnlineError"
 #define ONLINE_ERROR_NAMESPACE "errors.com.epicgames.online.generic"
diff --git a/Engine/Plugins/Online/OnlineSubsystem/Source/Private/OnlineSubsystem.cpp b/Engine/Plugins/Online/OnlineSubsystem/Source/Private/OnlineSubsystem.cpp
--- a/Engine/Plugins/Online/OnlineSubsystem/Source/Private/OnlineSubsystem.cpp
+++ b/Engine/Plugins/Online/OnlineSubsystem/Source/Private/OnlineSubsystem.cpp
@@ -6,6 +6,7 @@
 #include "Misc/TrackedActivity.h"
 #include "HAL/IConsoleManager.h"
 #include "Online/OnlineSessionNames.h"
+#include "OnlineIdentityErrors.h" // IWYU pragma: keep
 #include "OnlineSessionSettings.h"
 
 #include "Interfaces/OnlineChatInterface.h"
@@ -25,7 +26,6 @@
 #include "Interfaces/OnlineLeaderboardInterface.h"
 #include "Interfaces/OnlineTournamentInterface.h"
 #include "Interfaces/OnlineStatsInterface.h"
-#include "OnlineSubsystemNames.h"
 
 LLM_DEFINE_TAG(OnlineSubsystem);
 
diff --git a/Engine/Plugins/Online/OnlineSubsystem/Source/Public/Interfaces/OnlineAchievementsInterface.h b/Engine/Plugins/Online/OnlineSubsystem/Source/Public/Interfaces/OnlineAchievementsInterface.h
--- a/Engine/Plugins/Online/OnlineSubsystem/Source/Public/Interfaces/OnlineAchievementsInterface.h
+++ b/Engine/Plugins/Online/OnlineSubsystem/Source/Public/Interfaces/OnlineAchievementsInterface.h
@@ -1,7 +1,7 @@
 // Copyright Epic Games, Inc. All Rights Reserved.
 #pragma once
 
-#include "CoreMinimal.h"
+#include "Misc/DateTime.h"
 #include "OnlineDelegateMacros.h"
 #include "OnlineStats.h"
 
@@ -180,3 +180,7 @@ public:
 	 */
 	DEFINE_ONLINE_DELEGATE_TWO_PARAM(OnAchievementUnlocked, const FUniqueNetId&, const FString&);
 };
+
+#if UE_ENABLE_INCLUDE_ORDER_DEPRECATED_IN_5_2
+#include "CoreMinimal.h"
+#endif
diff --git a/Engine/Plugins/Online/OnlineSubsystem/Source/Public/Interfaces/OnlineChatInterface.h b/Engine/Plugins/Online/OnlineSubsystem/Source/Public/Interfaces/OnlineChatInterface.h
--- a/Engine/Plugins/Online/OnlineSubsystem/Source/Public/Interfaces/OnlineChatInterface.h
+++ b/Engine/Plugins/Online/OnlineSubsystem/Source/Public/Interfaces/OnlineChatInterface.h
@@ -2,9 +2,11 @@
 
 #pragma once
 
-#include "CoreMinimal.h"
-#include "Online/CoreOnline.h"
+#include "Delegates/Delegate.h"
 #include "OnlineDelegateMacros.h"
+#include "Online/CoreOnlineFwd.h"
+
+struct FDateTime;
 
 ONLINESUBSYSTEM_API DECLARE_LOG_CATEGORY_EXTERN(LogOnlineChat, Log, All);
 
@@ -394,3 +396,8 @@ public:
 	DEFINE_ONLINE_DELEGATE_THREE_PARAM(OnChatRoomMessageReceived, const FUniqueNetId&, const FChatRoomId&, const TSharedRef<FChatMessage>&);
 	DEFINE_ONLINE_DELEGATE_TWO_PARAM(OnChatPrivateMessageReceived, const FUniqueNetId&, const TSharedRef<FChatMessage>&);
 };
+
+#if UE_ENABLE_INCLUDE_ORDER_DEPRECATED_IN_5_2
+#include "CoreMinimal.h"
+#include "Online/CoreOnline.h"
+#endif
diff --git a/Engine/Plugins/Online/OnlineSubsystem/Source/Public/Interfaces/OnlineDeviceNotificationInterface.h b/Engine/Plugins/Online/OnlineSubsystem/Source/Public/Interfaces/OnlineDeviceNotificationInterface.h
--- a/Engine/Plugins/Online/OnlineSubsystem/Source/Public/Interfaces/OnlineDeviceNotificationInterface.h
+++ b/Engine/Plugins/Online/OnlineSubsystem/Source/Public/Interfaces/OnlineDeviceNotificationInterface.h
@@ -2,8 +2,10 @@
 
 #pragma once
 
-#include "CoreMinimal.h"
-#include "Online/CoreOnline.h"
+#include "Delegates/Delegate.h"
+
+class FString;
+class FUniqueNetId;
 
 struct FOnlineError;
 
@@ -49,3 +51,8 @@ public:
 	*/
 	virtual void UnregisterDeviceByChannelType(const FUniqueNetId& LocalUserId, const FString& InstanceId, const FString& ChannelType, const FOnUnregisterDeviceByChannelTypeComplete& Delegate) = 0;
 };
+
+#if UE_ENABLE_INCLUDE_ORDER_DEPRECATED_IN_5_2
+#include "CoreMinimal.h"
+#include "Online/CoreOnline.h"
+#endif
diff --git a/Engine/Plugins/Online/OnlineSubsystem/Source/Public/Interfaces/OnlineExternalUIInterface.h b/Engine/Plugins/Online/OnlineSubsystem/Source/Public/Interfaces/OnlineExternalUIInterface.h
--- a/Engine/Plugins/Online/OnlineSubsystem/Source/Public/Interfaces/OnlineExternalUIInterface.h
+++ b/Engine/Plugins/Online/OnlineSubsystem/Source/Public/Interfaces/OnlineExternalUIInterface.h
@@ -2,9 +2,6 @@
 
 #pragma once
 
-#include "CoreMinimal.h"
-#include "Online/CoreOnlineFwd.h"
-#include "OnlineDelegateMacros.h"
 #include "Interfaces/OnlineMessageInterface.h"
 #include "OnlineError.h"
 
@@ -467,3 +464,7 @@ public:
 
 typedef TSharedPtr<IOnlineExternalUI, ESPMode::ThreadSafe> IOnlineExternalUIPtr;
 
+
+#if UE_ENABLE_INCLUDE_ORDER_DEPRECATED_IN_5_2
+#include "CoreMinimal.h"
+#endif
diff --git a/Engine/Plugins/Online/OnlineSubsystem/Source/Public/Interfaces/OnlineGameActivityInterface.h b/Engine/Plugins/Online/OnlineSubsystem/Source/Public/Interfaces/OnlineGameActivityInterface.h
--- a/Engine/Plugins/Online/OnlineSubsystem/Source/Public/Interfaces/OnlineGameActivityInterface.h
+++ b/Engine/Plugins/Online/OnlineSubsystem/Source/Public/Interfaces/OnlineGameActivityInterface.h
@@ -2,10 +2,12 @@
 
 #pragma once
 
-#include "CoreMinimal.h"
-#include "OnlineEventsInterface.h"
+#include "Math/Vector.h"
 #include "OnlineDelegateMacros.h"
 
+class FOnlineEventParms;
+template <typename OptionalType> struct TOptional;
+
 class FUniqueNetId;
 struct FOnlineError;
 class FOnlineSessionSearchResult;
@@ -245,4 +247,9 @@ public:
 	 * @param ActivityPlayerLocation - zone id and player map coordinates (X and Y values)
 	 */
 	virtual void UpdatePlayerLocation(const FUniqueNetId& LocalUserId, TOptional<FOnlineActivityPlayerLocation>& ActivityPlayerLocation) = 0;
-};
\ No newline at end of file
+};
+
+#if UE_ENABLE_INCLUDE_ORDER_DEPRECATED_IN_5_2
+#include "CoreMinimal.h"
+#include "OnlineEventsInterface.h"
+#endif
diff --git a/Engine/Plugins/Online/OnlineSubsystem/Source/Public/Interfaces/OnlineGameItemStatsInterface.h b/Engine/Plugins/Online/OnlineSubsystem/Source/Public/Interfaces/OnlineGameItemStatsInterface.h
--- a/Engine/Plugins/Online/OnlineSubsystem/Source/Public/Interfaces/OnlineGameItemStatsInterface.h
+++ b/Engine/Plugins/Online/OnlineSubsystem/Source/Public/Interfaces/OnlineGameItemStatsInterface.h
@@ -2,9 +2,9 @@
 
 #pragma once
 
-#include "CoreMinimal.h"
-#include "OnlineError.h"
-#include "OnlineEventsInterface.h"
+#include "Delegates/Delegate.h"
+
+struct FOnlineError;
 
 class FUniqueNetId;
 
@@ -122,4 +122,10 @@ public:
 	 * @param CompletionDelegate - Item loadout change completion delegate called when ItemLoadoutChange call is complete
 	 */
 	virtual void ItemLoadoutChange(const FUniqueNetId& LocalUserId, const TArray<FString>& EquippedItems, const TArray<FString>& UnequippedItems, FOnItemLoadoutChangeComplete CompletionDelegate) = 0;
-};
\ No newline at end of file
+};
+
+#if UE_ENABLE_INCLUDE_ORDER_DEPRECATED_IN_5_2
+#include "CoreMinimal.h"
+#include "OnlineError.h"
+#include "OnlineEventsInterface.h"
+#endif
diff --git a/Engine/Plugins/Online/OnlineSubsystem/Source/Public/Interfaces/OnlineGameMatchesInterface.h b/Engine/Plugins/Online/OnlineSubsystem/Source/Public/Interfaces/OnlineGameMatchesInterface.h
--- a/Engine/Plugins/Online/OnlineSubsystem/Source/Public/Interfaces/OnlineGameMatchesInterface.h
+++ b/Engine/Plugins/Online/OnlineSubsystem/Source/Public/Interfaces/OnlineGameMatchesInterface.h
@@ -2,10 +2,9 @@
 
 #pragma once
 
-#include "CoreMinimal.h"
+#include "Delegates/Delegate.h"
 #include "Online/CoreOnlineFwd.h"
-#include "OnlineDelegateMacros.h"
-#include "OnlineStatsInterface.h"
+#include "OnlineKeyValuePair.h"
 
 struct FOnlineError;
 
@@ -407,4 +406,10 @@ public:
 	UE_DEPRECATED(5.0, "ProvideGameMatchFeedback is now deprecated and will be removed.")
 	virtual void ProvideGameMatchFeedback(const FUniqueNetId& UserId, const FString& MatchId, const bool bReviewTeam, const FOnGameMatchFeedbackComplete& CompletionDelegate) {}
 
-};
\ No newline at end of file
+};
+
+#if UE_ENABLE_INCLUDE_ORDER_DEPRECATED_IN_5_2
+#include "CoreMinimal.h"
+#include "OnlineDelegateMacros.h"
+#include "OnlineStatsInterface.h"
+#endif
diff --git a/Engine/Plugins/Online/OnlineSubsystem/Source/Public/Interfaces/OnlineGroupsInterface.h b/Engine/Plugins/Online/OnlineSubsystem/Source/Public/Interfaces/OnlineGroupsInterface.h
--- a/Engine/Plugins/Online/OnlineSubsystem/Source/Public/Interfaces/OnlineGroupsInterface.h
+++ b/Engine/Plugins/Online/OnlineSubsystem/Source/Public/Interfaces/OnlineGroupsInterface.h
@@ -1,8 +1,7 @@
 // Copyright Epic Games, Inc. All Rights Reserved.
 #pragma once
 
-#include "CoreMinimal.h"
-#include "Online/CoreOnline.h"
+#include "Misc/DateTime.h"
 #include "OnlineSubsystemTypes.h"
 
 class Error;
@@ -855,3 +854,7 @@ public:
 };
 
 typedef TSharedPtr<IOnlineGroups, ESPMode::ThreadSafe> IOnlineGroupsPtr;
+
+#if UE_ENABLE_INCLUDE_ORDER_DEPRECATED_IN_5_2
+#include "CoreMinimal.h"
+#endif
diff --git a/Engine/Plugins/Online/OnlineSubsystem/Source/Public/Interfaces/OnlineIdentityInterface.h b/Engine/Plugins/Online/OnlineSubsystem/Source/Public/Interfaces/OnlineIdentityInterface.h
--- a/Engine/Plugins/Online/OnlineSubsystem/Source/Public/Interfaces/OnlineIdentityInterface.h
+++ b/Engine/Plugins/Online/OnlineSubsystem/Source/Public/Interfaces/OnlineIdentityInterface.h
@@ -2,11 +2,10 @@
 
 #pragma once
 
-#include "CoreMinimal.h"
-#include "Online/CoreOnline.h"
 #include "OnlineSubsystemTypes.h"
 #include "OnlineDelegateMacros.h"
-#include "OnlineIdentityErrors.h"
+
+struct FOnlineError;
 
 
 
@@ -534,3 +533,8 @@ ONLINESUBSYSTEM_API FString ToDebugString(IOnlineIdentity::EPrivilegeResults Pri
 ONLINESUBSYSTEM_API FString ToDebugString(EUserPrivileges::Type UserPrivilege);
 
 ONLINESUBSYSTEM_API FString ToDebugString(const FControllerPairingChangedUserInfo& ControllerPairingChangedUserInfo);
+
+#if UE_ENABLE_INCLUDE_ORDER_DEPRECATED_IN_5_2
+#include "CoreMinimal.h"
+#include "OnlineIdentityErrors.h"
+#endif
diff --git a/Engine/Plugins/Online/OnlineSubsystem/Source/Public/Interfaces/OnlineLeaderboardInterface.h b/Engine/Plugins/Online/OnlineSubsystem/Source/Public/Interfaces/OnlineLeaderboardInterface.h
--- a/Engine/Plugins/Online/OnlineSubsystem/Source/Public/Interfaces/OnlineLeaderboardInterface.h
+++ b/Engine/Plugins/Online/OnlineSubsystem/Source/Public/Interfaces/OnlineLeaderboardInterface.h
@@ -2,8 +2,6 @@
 
 #pragma once
 
-#include "CoreMinimal.h"
-#include "Online/CoreOnline.h"
 #include "OnlineDelegateMacros.h"
 #include "OnlineStats.h"
 
@@ -158,3 +156,7 @@ public:
 
 typedef TSharedPtr<IOnlineLeaderboards, ESPMode::ThreadSafe> IOnlineLeaderboardsPtr;
 
+
+#if UE_ENABLE_INCLUDE_ORDER_DEPRECATED_IN_5_2
+#include "CoreMinimal.h"
+#endif
diff --git a/Engine/Plugins/Online/OnlineSubsystem/Source/Public/Interfaces/OnlineMessageInterface.h b/Engine/Plugins/Online/OnlineSubsystem/Source/Public/Interfaces/OnlineMessageInterface.h
--- a/Engine/Plugins/Online/OnlineSubsystem/Source/Public/Interfaces/OnlineMessageInterface.h
+++ b/Engine/Plugins/Online/OnlineSubsystem/Source/Public/Interfaces/OnlineMessageInterface.h
@@ -2,10 +2,8 @@
 
 #pragma once
 
-#include "CoreMinimal.h"
+#include "Online/CoreOnline.h"
 #include "OnlineKeyValuePair.h"
-#include "Online/CoreOnlineFwd.h"
-#include "OnlineSubsystemTypes.h"
 #include "OnlineDelegateMacros.h"
 
 class FJsonObject;
@@ -330,3 +328,8 @@ public:
 	 */
 	DEFINE_ONLINE_PLAYER_DELEGATE_THREE_PARAM(MAX_LOCAL_PLAYERS, OnDeleteMessageComplete, bool, const FUniqueMessageId&, const FString&);
 };
+
+#if UE_ENABLE_INCLUDE_ORDER_DEPRECATED_IN_5_2
+#include "CoreMinimal.h"
+#include "OnlineSubsystemTypes.h"
+#endif
diff --git a/Engine/Plugins/Online/OnlineSubsystem/Source/Public/Interfaces/OnlinePartyInterface.h b/Engine/Plugins/Online/OnlineSubsystem/Source/Public/Interfaces/OnlinePartyInterface.h
--- a/Engine/Plugins/Online/OnlineSubsystem/Source/Public/Interfaces/OnlinePartyInterface.h
+++ b/Engine/Plugins/Online/OnlineSubsystem/Source/Public/Interfaces/OnlinePartyInterface.h
@@ -2,9 +2,7 @@
 
 #pragma once
 
-#include "CoreMinimal.h"
 #include "OnlineKeyValuePair.h"
-#include "Online/CoreOnline.h"
 #include "OnlineSubsystemTypes.h"
 #include "OnlineDelegateMacros.h"
 #include "OnlineError.h"
@@ -2283,4 +2281,8 @@ ONLINESUBSYSTEM_API FString ToDebugString(const IOnlinePartyJoinInfo& JoinInfo);
 /** Dump key/value pairs for debugging */
 ONLINESUBSYSTEM_API FString ToDebugString(const FOnlineKeyValuePairs<FString, FVariantData>& KeyValAttrs);
 /** Dump state about the party data for debugging */
-ONLINESUBSYSTEM_API FString ToDebugString(const FOnlinePartyData& PartyData);
\ No newline at end of file
+ONLINESUBSYSTEM_API FString ToDebugString(const FOnlinePartyData& PartyData);
+
+#if UE_ENABLE_INCLUDE_ORDER_DEPRECATED_IN_5_2
+#include "CoreMinimal.h"
+#endif
diff --git a/Engine/Plugins/Online/OnlineSubsystem/Source/Public/Interfaces/OnlinePresenceInterface.h b/Engine/Plugins/Online/OnlineSubsystem/Source/Public/Interfaces/OnlinePresenceInterface.h
--- a/Engine/Plugins/Online/OnlineSubsystem/Source/Public/Interfaces/OnlinePresenceInterface.h
+++ b/Engine/Plugins/Online/OnlineSubsystem/Source/Public/Interfaces/OnlinePresenceInterface.h
@@ -2,9 +2,9 @@
 
 #pragma once
 
-#include "CoreMinimal.h"
-#include "Online/CoreOnline.h"
+#include "Internationalization/Internationalization.h"
 #include "OnlineSubsystemTypes.h"
+#include "Misc/DateTime.h"
 #include "OnlineDelegateMacros.h"
 #include "OnlineKeyValuePair.h"
 
@@ -380,3 +380,7 @@ public:
 	*/
 	virtual EOnlineCachedResult::Type GetCachedPresenceForApp(const FUniqueNetId& LocalUserId, const FUniqueNetId& User, const FString& AppId, TSharedPtr<FOnlineUserPresence>& OutPresence) = 0;
 };
+
+#if UE_ENABLE_INCLUDE_ORDER_DEPRECATED_IN_5_2
+#include "CoreMinimal.h"
+#endif
diff --git a/Engine/Plugins/Online/OnlineSubsystem/Source/Public/Interfaces/OnlinePurchaseInterface.h b/Engine/Plugins/Online/OnlineSubsystem/Source/Public/Interfaces/OnlinePurchaseInterface.h
--- a/Engine/Plugins/Online/OnlineSubsystem/Source/Public/Interfaces/OnlinePurchaseInterface.h
+++ b/Engine/Plugins/Online/OnlineSubsystem/Source/Public/Interfaces/OnlinePurchaseInterface.h
@@ -2,11 +2,11 @@
 
 #pragma once
 
-#include "CoreMinimal.h"
 #include "Interfaces/OnlineEntitlementsInterface.h"
 #include "Interfaces/OnlineStoreInterfaceV2.h"
 #include "OnlineJsonSerializer.h"
-#include "OnlineError.h"
+
+struct FOnlineError;
 
 #define PURCHASE_REDEMPTION_CODE_USED "errors.com.epicgames.coderedemption.code_used"
 #define PURCHASE_REDEMPTION_CODE_NOT_FOUND "errors.com.epicgames.coderedemption.code_not_found"
@@ -360,4 +360,9 @@ inline const TCHAR* LexToString(EPurchaseTransactionState State)
 	case EPurchaseTransactionState::Invalid: return TEXT("Invalid");
 	default: return TEXT("Unknown");
 	}
-}
\ No newline at end of file
+}
+
+#if UE_ENABLE_INCLUDE_ORDER_DEPRECATED_IN_5_2
+#include "CoreMinimal.h"
+#include "OnlineError.h"
+#endif
diff --git a/Engine/Plugins/Online/OnlineSubsystem/Source/Public/Interfaces/OnlineSessionInterface.h b/Engine/Plugins/Online/OnlineSubsystem/Source/Public/Interfaces/OnlineSessionInterface.h
--- a/Engine/Plugins/Online/OnlineSubsystem/Source/Public/Interfaces/OnlineSessionInterface.h
+++ b/Engine/Plugins/Online/OnlineSubsystem/Source/Public/Interfaces/OnlineSessionInterface.h
@@ -2,8 +2,6 @@
 
 #pragma once
 
-#include "CoreMinimal.h"
-#include "Online/CoreOnline.h"
 #include "OnlineSubsystemTypes.h"
 #include "OnlineDelegateMacros.h"
 #include "OnlineSessionDelegates.h"
@@ -964,3 +962,7 @@ public:
 
 typedef TSharedPtr<IOnlineSession, ESPMode::ThreadSafe> IOnlineSessionPtr;
 
+
+#if UE_ENABLE_INCLUDE_ORDER_DEPRECATED_IN_5_2
+#include "CoreMinimal.h"
+#endif
diff --git a/Engine/Plugins/Online/OnlineSubsystem/Source/Public/Interfaces/OnlineSharedCloudInterface.h b/Engine/Plugins/Online/OnlineSubsystem/Source/Public/Interfaces/OnlineSharedCloudInterface.h
--- a/Engine/Plugins/Online/OnlineSubsystem/Source/Public/Interfaces/OnlineSharedCloudInterface.h
+++ b/Engine/Plugins/Online/OnlineSubsystem/Source/Public/Interfaces/OnlineSharedCloudInterface.h
@@ -2,10 +2,12 @@
  
 #pragma once
 
-#include "CoreMinimal.h"
-#include "OnlineSubsystemTypes.h"
+#include "Delegates/Delegate.h"
 #include "OnlineDelegateMacros.h"
 
+class FSharedContentHandle;
+class FUniqueNetId;
+
 /**
  * Delegate fired when a shared file read from the network platform's storage is complete
  *
@@ -107,3 +109,8 @@ public:
 };
 
 typedef TSharedPtr<IOnlineSharedCloud, ESPMode::ThreadSafe> IOnlineSharedCloudPtr;
+
+#if UE_ENABLE_INCLUDE_ORDER_DEPRECATED_IN_5_2
+#include "CoreMinimal.h"
+#include "OnlineSubsystemTypes.h"
+#endif
diff --git a/Engine/Plugins/Online/OnlineSubsystem/Source/Public/Interfaces/OnlineStatsInterface.h b/Engine/Plugins/Online/OnlineSubsystem/Source/Public/Interfaces/OnlineStatsInterface.h
--- a/Engine/Plugins/Online/OnlineSubsystem/Source/Public/Interfaces/OnlineStatsInterface.h
+++ b/Engine/Plugins/Online/OnlineSubsystem/Source/Public/Interfaces/OnlineStatsInterface.h
@@ -1,10 +1,8 @@
 // Copyright Epic Games, Inc. All Rights Reserved.
 #pragma once
 
-#include "CoreMinimal.h"
-#include "OnlineSubsystemTypes.h"
+#include "Online/CoreOnline.h"
 #include "OnlineKeyValuePair.h"
-#include "Misc/Optional.h"
 
 struct FOnlineError;
 
@@ -310,3 +308,8 @@ public:
 
 };
 
+
+#if UE_ENABLE_INCLUDE_ORDER_DEPRECATED_IN_5_2
+#include "CoreMinimal.h"
+#include "OnlineSubsystemTypes.h"
+#endif
diff --git a/Engine/Plugins/Online/OnlineSubsystem/Source/Public/Interfaces/OnlineStoreInterfaceV2.h b/Engine/Plugins/Online/OnlineSubsystem/Source/Public/Interfaces/OnlineStoreInterfaceV2.h
--- a/Engine/Plugins/Online/OnlineSubsystem/Source/Public/Interfaces/OnlineStoreInterfaceV2.h
+++ b/Engine/Plugins/Online/OnlineSubsystem/Source/Public/Interfaces/OnlineStoreInterfaceV2.h
@@ -2,7 +2,7 @@
 
 #pragma once
 
-#include "CoreMinimal.h"
+#include "Misc/DateTime.h"
 #include "OnlineDelegateMacros.h"
 #include "OnlineStoreInterfaceV2.generated.h"
 
@@ -275,3 +275,7 @@ public:
 	 */
 	virtual TSharedPtr<FOnlineStoreOffer> GetOffer(const FUniqueOfferId& OfferId) const = 0;
 };
+
+#if UE_ENABLE_INCLUDE_ORDER_DEPRECATED_IN_5_2
+#include "CoreMinimal.h"
+#endif
diff --git a/Engine/Plugins/Online/OnlineSubsystem/Source/Public/Interfaces/OnlineTournamentInterface.h b/Engine/Plugins/Online/OnlineSubsystem/Source/Public/Interfaces/OnlineTournamentInterface.h
--- a/Engine/Plugins/Online/OnlineSubsystem/Source/Public/Interfaces/OnlineTournamentInterface.h
+++ b/Engine/Plugins/Online/OnlineSubsystem/Source/Public/Interfaces/OnlineTournamentInterface.h
@@ -2,11 +2,13 @@
 
 #pragma once
 
-#include "CoreMinimal.h"
+#include "Delegates/Delegate.h"
 #include "Online/CoreOnlineFwd.h"
-#include "OnlineSubsystemTypes.h"
 #include "OnlineKeyValuePair.h"
 
+struct FDateTime;
+struct FTimespan;
+
 ONLINESUBSYSTEM_API DECLARE_LOG_CATEGORY_EXTERN(LogOnlineTournament, Display, All);
 
 #define UE_LOG_ONLINE_TOURNAMENT(Verbosity, Format, ...) \
@@ -809,3 +811,8 @@ public:
 	virtual void DumpCachedTeamInfo(const FUniqueNetIdRef UserId) const = 0;
 #endif // !UE_BUILD_SHIPPING
 };
+
+#if UE_ENABLE_INCLUDE_ORDER_DEPRECATED_IN_5_2
+#include "CoreMinimal.h"
+#include "OnlineSubsystemTypes.h"
+#endif
diff --git a/Engine/Plugins/Online/OnlineSubsystem/Source/Public/Interfaces/OnlineTracingInterface.h b/Engine/Plugins/Online/OnlineSubsystem/Source/Public/Interfaces/OnlineTracingInterface.h
--- a/Engine/Plugins/Online/OnlineSubsystem/Source/Public/Interfaces/OnlineTracingInterface.h
+++ b/Engine/Plugins/Online/OnlineSubsystem/Source/Public/Interfaces/OnlineTracingInterface.h
@@ -2,9 +2,11 @@
 
 #pragma once
 
-#include "CoreMinimal.h"
+#include "Containers/ContainersFwd.h" // IWYU pragma: keep
 #include "Features/IModularFeature.h"
 
+class FName;
+
 class FOutputDevice;
 
 /**
@@ -40,3 +42,8 @@ protected:
 
 	virtual void EndContextImpl(FName ContextName) = 0;
 };
+
+#if UE_ENABLE_INCLUDE_ORDER_DEPRECATED_IN_5_2
+#include "Containers/Array.h"
+#include "CoreMinimal.h"
+#endif
diff --git a/Engine/Plugins/Online/OnlineSubsystem/Source/Public/Interfaces/OnlineUserCloudInterface.h b/Engine/Plugins/Online/OnlineSubsystem/Source/Public/Interfaces/OnlineUserCloudInterface.h
--- a/Engine/Plugins/Online/OnlineSubsystem/Source/Public/Interfaces/OnlineUserCloudInterface.h
+++ b/Engine/Plugins/Online/OnlineSubsystem/Source/Public/Interfaces/OnlineUserCloudInterface.h
@@ -2,10 +2,13 @@
 
 #pragma once
 
-#include "CoreMinimal.h"
-#include "OnlineSubsystemTypes.h"
+#include "Delegates/Delegate.h"
 #include "OnlineDelegateMacros.h"
 
+class FUniqueNetId;
+struct FCloudFileHeader;
+template <typename OptionalType> struct TOptional;
+
 ONLINESUBSYSTEM_API DECLARE_LOG_CATEGORY_EXTERN(LogOnlineCloud, Log, All);
 
 #define UE_LOG_ONLINE_CLOUD(Verbosity, Format, ...) \
@@ -308,3 +311,8 @@ public:
 };
 
 typedef TSharedPtr<IOnlineUserCloud, ESPMode::ThreadSafe> IOnlineUserCloudPtr;
+
+#if UE_ENABLE_INCLUDE_ORDER_DEPRECATED_IN_5_2
+#include "CoreMinimal.h"
+#include "OnlineSubsystemTypes.h"
+#endif
diff --git a/Engine/Plugins/Online/OnlineSubsystem/Source/Public/Interfaces/OnlineUserInterface.h b/Engine/Plugins/Online/OnlineSubsystem/Source/Public/Interfaces/OnlineUserInterface.h
--- a/Engine/Plugins/Online/OnlineSubsystem/Source/Public/Interfaces/OnlineUserInterface.h
+++ b/Engine/Plugins/Online/OnlineSubsystem/Source/Public/Interfaces/OnlineUserInterface.h
@@ -2,11 +2,11 @@
 
 #pragma once
 
-#include "CoreMinimal.h"
 #include "Online/CoreOnline.h"
-#include "OnlineSubsystemTypes.h"
 #include "OnlineDelegateMacros.h"
 
+class FOnlineUser;
+
 ONLINESUBSYSTEM_API DECLARE_LOG_CATEGORY_EXTERN(LogOnlineUser, Log, All);
 #define UE_LOG_ONLINE_USER(Verbosity, Format, ...) \
 { \
@@ -164,3 +164,8 @@ public:
 	 */
 	virtual FUniqueNetIdPtr GetExternalIdMapping(const FExternalIdQueryOptions& QueryOptions, const FString& ExternalId) = 0;
 };
+
+#if UE_ENABLE_INCLUDE_ORDER_DEPRECATED_IN_5_2
+#include "CoreMinimal.h"
+#include "OnlineSubsystemTypes.h"
+#endif
diff --git a/Engine/Plugins/Online/OnlineSubsystem/Source/Public/NamedInterfaces.h b/Engine/Plugins/Online/OnlineSubsystem/Source/Public/NamedInterfaces.h
--- a/Engine/Plugins/Online/OnlineSubsystem/Source/Public/NamedInterfaces.h
+++ b/Engine/Plugins/Online/OnlineSubsystem/Source/Public/NamedInterfaces.h
@@ -6,9 +6,6 @@
 
 #pragma once
 
-#include "CoreMinimal.h"
-#include "UObject/ObjectMacros.h"
-#include "UObject/Object.h"
 #include "NamedInterfaces.generated.h"
 
 /** Holds a named object interface for dynamically bound interfaces */
@@ -97,3 +94,7 @@ private:
 	FOnCleanup CleanupDelegates;
 };
 
+
+#if UE_ENABLE_INCLUDE_ORDER_DEPRECATED_IN_5_2
+#include "CoreMinimal.h"
+#endif
diff --git a/Engine/Plugins/Online/OnlineSubsystem/Source/Public/NboSerializerOSS.h b/Engine/Plugins/Online/OnlineSubsystem/Source/Public/NboSerializerOSS.h
--- a/Engine/Plugins/Online/OnlineSubsystem/Source/Public/NboSerializerOSS.h
+++ b/Engine/Plugins/Online/OnlineSubsystem/Source/Public/NboSerializerOSS.h
@@ -5,7 +5,6 @@
 #include "Online/NboSerializer.h"
 
 #include "OnlineSessionSettings.h"
-#include "OnlineKeyValuePair.h"
 
 /**
  * Serializes data in network byte order form into a buffer
@@ -267,4 +266,4 @@ public:
 
 		return Ar;
 	}
-};
\ No newline at end of file
+};
diff --git a/Engine/Plugins/Online/OnlineSubsystem/Source/Public/Online.h b/Engine/Plugins/Online/OnlineSubsystem/Source/Public/Online.h
--- a/Engine/Plugins/Online/OnlineSubsystem/Source/Public/Online.h
+++ b/Engine/Plugins/Online/OnlineSubsystem/Source/Public/Online.h
@@ -2,26 +2,16 @@
 
 #pragma once
 
-#include "CoreMinimal.h"
 #include "OnlineSubsystem.h"
 
-#include "OnlineSessionSettings.h"
-#include "Interfaces/OnlineChatInterface.h"
 #include "Interfaces/OnlineIdentityInterface.h"
 #include "Interfaces/OnlineFriendsInterface.h"
-#include "Interfaces/OnlineEventsInterface.h"
 #include "Interfaces/OnlineSessionInterface.h"
-#include "Interfaces/OnlinePartyInterface.h"
 #include "Interfaces/OnlineExternalUIInterface.h"
 #include "Interfaces/VoiceInterface.h"
 #include "Interfaces/OnlineTitleFileInterface.h"
-#include "Interfaces/OnlineAchievementsInterface.h"
-#include "Interfaces/OnlinePresenceInterface.h"
-#include "Interfaces/OnlinePurchaseInterface.h"
 #include "Interfaces/OnlineEntitlementsInterface.h"
 #include "Interfaces/OnlineUserCloudInterface.h"
-#include "Interfaces/OnlineUserInterface.h"
-#include "Interfaces/OnlineStatsInterface.h"
 
 /** Macro to handle the boilerplate of accessing the proper online subsystem and getting the requested interface */
 #define IMPLEMENT_GET_INTERFACE(InterfaceType) \
@@ -175,3 +165,16 @@ namespace Online
 };
 
 #undef IMPLEMENT_GET_INTERFACE
+
+#if UE_ENABLE_INCLUDE_ORDER_DEPRECATED_IN_5_2
+#include "CoreMinimal.h"
+#include "Interfaces/OnlineAchievementsInterface.h"
+#include "Interfaces/OnlineChatInterface.h"
+#include "Interfaces/OnlineEventsInterface.h"
+#include "Interfaces/OnlinePartyInterface.h"
+#include "Interfaces/OnlinePresenceInterface.h"
+#include "Interfaces/OnlinePurchaseInterface.h"
+#include "Interfaces/OnlineStatsInterface.h"
+#include "Interfaces/OnlineUserInterface.h"
+#include "OnlineSessionSettings.h"
+#endif
diff --git a/Engine/Plugins/Online/OnlineSubsystem/Source/Public/OnlineAsyncTaskManager.h b/Engine/Plugins/Online/OnlineSubsystem/Source/Public/OnlineAsyncTaskManager.h
--- a/Engine/Plugins/Online/OnlineSubsystem/Source/Public/OnlineAsyncTaskManager.h
+++ b/Engine/Plugins/Online/OnlineSubsystem/Source/Public/OnlineAsyncTaskManager.h
@@ -2,7 +2,7 @@
 
 #pragma once
 
-#include "CoreMinimal.h"
+#include "CoreGlobals.h"
 #include "HAL/PlatformTime.h"
 #include "HAL/PlatformProcess.h"
 #include "HAL/Runnable.h"
@@ -10,7 +10,6 @@
 #include "Misc/SingleThreadRunnable.h"
 #include "OnlineSubsystemPackage.h"
 #include "Containers/Queue.h"
-#include <atomic>
 
 /**
  * Base class of any async task that can be returned to the game thread by the async task manager
@@ -500,3 +499,7 @@ public:
 };
 
 
+
+#if UE_ENABLE_INCLUDE_ORDER_DEPRECATED_IN_5_2
+#include "CoreMinimal.h"
+#endif
diff --git a/Engine/Plugins/Online/OnlineSubsystem/Source/Public/OnlineError.h b/Engine/Plugins/Online/OnlineSubsystem/Source/Public/OnlineError.h
--- a/Engine/Plugins/Online/OnlineSubsystem/Source/Public/OnlineError.h
+++ b/Engine/Plugins/Online/OnlineSubsystem/Source/Public/OnlineError.h
@@ -2,7 +2,7 @@
 
 #pragma once
 
-#include "CoreMinimal.h"
+#include "Internationalization/Internationalization.h" // IWYU pragma: keep
 
 #define ONLINE_ERROR_LEGACY 1
 
@@ -192,3 +192,7 @@ protected:
 /** must be defined to a valid namespace for using ONLINE_ERROR factory macro */
 #undef ONLINE_ERROR_NAMESPACE
 #define ONLINE_ERROR(...) FOnlineError::CreateError(TEXT(ONLINE_ERROR_NAMESPACE), __VA_ARGS__)
+
+#if UE_ENABLE_INCLUDE_ORDER_DEPRECATED_IN_5_2
+#include "CoreMinimal.h"
+#endif
diff --git a/Engine/Plugins/Online/OnlineSubsystem/Source/Public/OnlineFriendErrors.h b/Engine/Plugins/Online/OnlineSubsystem/Source/Public/OnlineFriendErrors.h
--- a/Engine/Plugins/Online/OnlineSubsystem/Source/Public/OnlineFriendErrors.h
+++ b/Engine/Plugins/Online/OnlineSubsystem/Source/Public/OnlineFriendErrors.h
@@ -3,7 +3,7 @@
 #pragma once
 
 
-#include "OnlineError.h"
+#include "OnlineError.h" // IWYU pragma: keep
 #define LOCTEXT_NAMESPACE "OnlineFriend"
 #define ONLINE_ERROR_NAMESPACE "errors.com.epicgames.oss.friend"
 
