Remove unnecessary public include adding the root module directory for several online plugins

#rnx

diff --git a/Engine/Plugins/Online/OnlineServices/Source/OnlineServicesCommon/OnlineServicesCommon.Build.cs b/Engine/Plugins/Online/OnlineServices/Source/OnlineServicesCommon/OnlineServicesCommon.Build.cs
--- a/Engine/Plugins/Online/OnlineServices/Source/OnlineServicesCommon/OnlineServicesCommon.Build.cs
+++ b/Engine/Plugins/Online/OnlineServices/Source/OnlineServicesCommon/OnlineServicesCommon.Build.cs
@@ -17,8 +17,6 @@ public class OnlineServicesCommon : ModuleRules
 			}
 		);
 
-		PublicIncludePaths.Add(ModuleDirectory);
-
 		// OnlineService cannot depend on Engine!
 		PrivateDependencyModuleNames.AddRange(
 			new string[] { 
diff --git a/Engine/Plugins/Online/OnlineServices/Source/OnlineServicesInterface/OnlineServicesInterface.Build.cs b/Engine/Plugins/Online/OnlineServices/Source/OnlineServicesInterface/OnlineServicesInterface.Build.cs
--- a/Engine/Plugins/Online/OnlineServices/Source/OnlineServicesInterface/OnlineServicesInterface.Build.cs
+++ b/Engine/Plugins/Online/OnlineServices/Source/OnlineServicesInterface/OnlineServicesInterface.Build.cs
@@ -15,8 +15,6 @@ public class OnlineServicesInterface : ModuleRules
 			}
 		);
 
-		PublicIncludePaths.Add(ModuleDirectory);
-
 		// OnlineService cannot depend on Engine!
 		PrivateDependencyModuleNames.AddRange(
 			new string[] { 
diff --git a/Engine/Plugins/Online/OnlineSubsystem/Source/OnlineSubsystem.Build.cs b/Engine/Plugins/Online/OnlineSubsystem/Source/OnlineSubsystem.Build.cs
--- a/Engine/Plugins/Online/OnlineSubsystem/Source/OnlineSubsystem.Build.cs
+++ b/Engine/Plugins/Online/OnlineSubsystem/Source/OnlineSubsystem.Build.cs
@@ -1,5 +1,6 @@
 // Copyright Epic Games, Inc. All Rights Reserved.
 
+using System.IO;
 using UnrealBuildTool;
 
 public class OnlineSubsystem : ModuleRules
@@ -7,6 +8,9 @@ public class OnlineSubsystem : ModuleRules
 	public OnlineSubsystem(ReadOnlyTargetRules Target) : base(Target)
     {
 		PCHUsage = ModuleRules.PCHUsageMode.UseExplicitOrSharedPCHs;
+
+		// Required for Test cpp files in OnlineSubsystemMcp
+		PublicIncludePaths.Add(Path.Join(ModuleDirectory, "Test"));
 		
 		PublicDependencyModuleNames.AddRange(
 			new string[] {
@@ -17,8 +21,6 @@ public class OnlineSubsystem : ModuleRules
 			}
 		);
 
-		PublicIncludePaths.Add(ModuleDirectory);
-
         PublicDefinitions.Add("ONLINESUBSYSTEM_PACKAGE=1");
 		PublicDefinitions.Add("DEBUG_LAN_BEACON=0");
 
