Remove deprecated party functions

diff --git a/Engine/Plugins/Online/OnlineFramework/Source/Party/Private/Party/SocialParty.cpp b/Engine/Plugins/Online/OnlineFramework/Source/Party/Private/Party/SocialParty.cpp
--- a/Engine/Plugins/Online/OnlineFramework/Source/Party/Private/Party/SocialParty.cpp
+++ b/Engine/Plugins/Online/OnlineFramework/Source/Party/Private/Party/SocialParty.cpp
@@ -486,11 +486,6 @@ void USocialParty::InitializePartyInternal()
 	PartyInterface->AddOnPartyMemberExitedDelegate_Handle(FOnPartyMemberExitedDelegate::CreateUObject(this, &USocialParty::HandlePartyMemberExited));
 	PartyInterface->AddOnPartySystemStateChangeDelegate_Handle(FOnPartySystemStateChangeDelegate::CreateUObject(this, &USocialParty::HandlePartySystemStateChange));
 
-	PRAGMA_DISABLE_DEPRECATION_WARNINGS
-	PartyInterface->AddOnPartyJIPRequestReceivedDelegate_Handle(FOnPartyJIPRequestReceivedDelegate::CreateUObject(this, &USocialParty::HandlePartyJIPRequestReceived));
-	PartyInterface->AddOnPartyJIPResponseDelegate_Handle(FOnPartyJIPResponseDelegate::CreateUObject(this, &USocialParty::HandlePartyMemberJIP));
-	PRAGMA_ENABLE_DEPRECATION_WARNINGS
-
 	// Create a UPartyMember for every existing member on the OSS party
 	TArray<FOnlinePartyMemberConstRef> OssPartyMembers;
 	PartyInterface->GetPartyMembers(*OwningLocalUserId, GetPartyId(), OssPartyMembers);
@@ -809,58 +804,6 @@ void USocialParty::RemovePlayerFromReservationBeacon(const FUniqueNetId& LocalUs
 	}
 }
 
-// DEPRECATED - Use the new join in progress flow with USocialParty::RequestJoinInProgress.
-void USocialParty::HandlePartyJIPRequestReceived(const FUniqueNetId& LocalUserId, const FOnlinePartyId& PartyId, const FUniqueNetId& SenderId)
-{
-	if (!IsLocalPlayerPartyLeader() || PartyId != GetPartyId())
-	{
-		return;
-	}
-
-	IOnlinePartyPtr PartyInterface = Online::GetPartyInterfaceChecked(GetWorld());
-
-	PRAGMA_DISABLE_DEPRECATION_WARNINGS
-	FPartyJoinApproval JoinApproval = EvaluateJIPRequest(SenderId);
-	PRAGMA_ENABLE_DEPRECATION_WARNINGS
-
-	if (JoinApproval.GetApprovalAction() == EApprovalAction::Enqueue ||
-		JoinApproval.GetApprovalAction() == EApprovalAction::EnqueueAndStartBeacon)
-	{
-		FUserPlatform MemberPlatform;
-		for (const UPartyMember* Member : GetPartyMembers())
-		{
-			if (Member->GetPrimaryNetId() == SenderId)
-			{
-				MemberPlatform = Member->GetRepData().GetPlatformDataPlatform();
-				break;
-			}
-		}
-
-		// Enqueue for a more opportune time
-		UE_LOG(LogParty, Verbose, TEXT("[%s] Enqueuing JIP approval request for %s"), *PartyId.ToString(), *SenderId.ToString());
-
-		FPendingMemberApproval PendingApproval;
-		PendingApproval.RecipientId.SetUniqueNetId(LocalUserId.AsShared());
-		PendingApproval.Members.Emplace(SenderId.AsShared(), MoveTemp(MemberPlatform));
-		PendingApproval.bIsJIPApproval = true;
-		PendingApprovals.Enqueue(MoveTemp(PendingApproval));
-
-		if (!ReservationBeaconClient.Get() && JoinApproval.GetApprovalAction() == EApprovalAction::EnqueueAndStartBeacon)
-		{
-			ConnectToReservationBeacon();
-		}
-	}
-	else
-	{
-		const bool bIsApproved = JoinApproval.CanJoin();
-		UE_LOG(LogParty, Log, TEXT("[%s] Responding to approval request for %s with %s"), *PartyId.ToString(), *SenderId.ToString(), bIsApproved ? TEXT("approved") : TEXT("denied"));
-
-		PRAGMA_DISABLE_DEPRECATION_WARNINGS
-		PartyInterface->ApproveJIPRequest(LocalUserId, PartyId, SenderId, bIsApproved, JoinApproval.GetDenialReason());
-		PRAGMA_ENABLE_DEPRECATION_WARNINGS
-	}
-}
-
 void USocialParty::HandleJoinabilityQueryReceived(const FUniqueNetId& LocalUserId, const FOnlinePartyId& PartyId, const IOnlinePartyPendingJoinRequestInfo& JoinRequestInfo)
 {
 	if (PartyId == GetPartyId()) 
@@ -969,24 +912,6 @@ void USocialParty::HandlePartyMemberJoined(const FUniqueNetId& LocalUserId, cons
 	}
 }
 
-// DEPRECATED - Use the new join in progress flow with USocialParty::RequestJoinInProgress.
-void USocialParty::HandlePartyMemberJIP(const FUniqueNetId& LocalUserId, const FOnlinePartyId& PartyId, bool Success, int32 DeniedResultCode)
-{
-	if (PartyId == GetPartyId())
-	{
-		FString DeniedResultCodeString = StaticEnum<EPartyJoinDenialReason>()->GetNameStringByValue(DeniedResultCode);
-		if (DeniedResultCodeString.IsEmpty())
-		{
-			UE_LOG(LogParty, Warning, TEXT("Failed to convert JIP result code. Value=%d"), DeniedResultCode);
-			DeniedResultCodeString = TEXT("Invalid");
-		}
-
-		// We are allowed to join the party.. start the JIP flow. 
-		OnPartyJIPApprovedEvent.Broadcast(PartyId, Success);
-		OnPartyJIPResponseEvent.Broadcast(PartyId, Success, DeniedResultCodeString);
-	}
-}
-
 void USocialParty::HandlePartyMemberPromoted(const FUniqueNetId& LocalUserId, const FOnlinePartyId& PartyId, const FUniqueNetId& NewLeaderId)
 {
 	if (PartyId == GetPartyId())
@@ -1609,10 +1534,6 @@ void USocialParty::RejectAllPendingJoinRequests()
 		UE_LOG(LogParty, Log, TEXT("[%s] Responding to approval request for %s with %s"), *PartyId.ToString(), *PrimaryJoiningUserId.ToDebugString(), ToString(EPartyJoinDenialReason::Busy));
 		if (PendingApproval.bIsJIPApproval)
 		{
-			PRAGMA_DISABLE_DEPRECATION_WARNINGS
-			PartyInterface->ApproveJIPRequest(*PendingApproval.RecipientId, PartyId, PrimaryJoiningUserId, false, (int32)EPartyJoinDenialReason::Busy);
-			PRAGMA_ENABLE_DEPRECATION_WARNINGS
-
 			RespondToJoinInProgressRequest(PendingApproval, EPartyJoinDenialReason::Busy);
 		}
 		else
@@ -1733,11 +1654,7 @@ void USocialParty::HandleReservationRequestComplete(EPartyReservationResult::Typ
 			IOnlinePartyPtr PartyInterface = Online::GetPartyInterfaceChecked(GetWorld());
 			if (PendingApproval.bIsJIPApproval)
 			{
-				// This player is already in our party. ApproveJIPRequest
-				PRAGMA_DISABLE_DEPRECATION_WARNINGS
-				PartyInterface->ApproveJIPRequest(*PendingApproval.RecipientId, GetPartyId(), *PendingApproval.Members[0].MemberId, bReservationApproved, DenialReason);
-				PRAGMA_ENABLE_DEPRECATION_WARNINGS
-
+				// Respond to the JIP request
 				RespondToJoinInProgressRequest(PendingApproval, DenialReason.GetReason());
 			}
 			else if (PendingApproval.bIsPlayerRemoval)
diff --git a/Engine/Plugins/Online/OnlineFramework/Source/Party/Private/SocialDebugTools.cpp b/Engine/Plugins/Online/OnlineFramework/Source/Party/Private/SocialDebugTools.cpp
--- a/Engine/Plugins/Online/OnlineFramework/Source/Party/Private/SocialDebugTools.cpp
+++ b/Engine/Plugins/Online/OnlineFramework/Source/Party/Private/SocialDebugTools.cpp
@@ -799,7 +799,7 @@ void USocialDebugTools::FInstanceContext::Shutdown()
 		IOnlinePartyPtr OnlineParty = OnlineSub->GetPartyInterface();
 		if (OnlineParty.IsValid())
 		{
-			OnlineParty->ClearOnPartyInviteReceivedDelegate_Handle(PartyInviteReceivedDelegateHandle);
+			OnlineParty->ClearOnPartyInviteReceivedExDelegate_Handle(PartyInviteReceivedDelegateHandle);
 			OnlineParty->ClearOnPartyJoinRequestReceivedDelegate_Handle(PartyJoinRequestReceivedDelegateHandle);
 			OnlineParty.Reset();
 		}
diff --git a/Engine/Plugins/Online/OnlineFramework/Source/Party/Private/User/SocialUser.cpp b/Engine/Plugins/Online/OnlineFramework/Source/Party/Private/User/SocialUser.cpp
--- a/Engine/Plugins/Online/OnlineFramework/Source/Party/Private/User/SocialUser.cpp
+++ b/Engine/Plugins/Online/OnlineFramework/Source/Party/Private/User/SocialUser.cpp
@@ -1126,15 +1126,8 @@ void USocialUser::HandlePartyRequestToJoinSent(const FUniqueNetId& LocalUserId,
 
 	if (Result == ERequestToJoinPartyCompletionResult::Succeeded)
 	{
-PRAGMA_DISABLE_DEPRECATION_WARNINGS
-		NotifyRequestToJoinSent(ExpiresAt);
-PRAGMA_ENABLE_DEPRECATION_WARNINGS
-
 		GetOwningToolkit().OnPartyRequestToJoinSent().Broadcast(*this);
 	}
-PRAGMA_DISABLE_DEPRECATION_WARNINGS
-	GetOwningToolkit().OnRequestToJoinPartyComplete(PartyLeaderId, Result);
-PRAGMA_ENABLE_DEPRECATION_WARNINGS
 }
 
 TSharedPtr<const IOnlinePartyJoinInfo> USocialUser::GetPartyJoinInfo(const FOnlinePartyTypeId& PartyTypeId) const
diff --git a/Engine/Plugins/Online/OnlineFramework/Source/Party/Public/Party/SocialParty.h b/Engine/Plugins/Online/OnlineFramework/Source/Party/Public/Party/SocialParty.h
--- a/Engine/Plugins/Online/OnlineFramework/Source/Party/Public/Party/SocialParty.h
+++ b/Engine/Plugins/Online/OnlineFramework/Source/Party/Public/Party/SocialParty.h
@@ -353,14 +353,10 @@ private:	// Handlers
 	void HandlePartyDataReceived(const FUniqueNetId& LocalUserId, const FOnlinePartyId& PartyId, const FName& Namespace, const FOnlinePartyData& PartyData);
 	void HandleJoinabilityQueryReceived(const FUniqueNetId& LocalUserId, const FOnlinePartyId& PartyId, const IOnlinePartyPendingJoinRequestInfo& JoinRequestInfo);
 	void HandlePartyJoinRequestReceived(const FUniqueNetId& LocalUserId, const FOnlinePartyId& PartyId, const IOnlinePartyPendingJoinRequestInfo& JoinRequestInfo);
-	UE_DEPRECATED(5.1, "Use the new join in progress flow with USocialParty::RequestJoinInProgress.")
-	void HandlePartyJIPRequestReceived(const FUniqueNetId& LocalUserId, const FOnlinePartyId& PartyId, const FUniqueNetId& SenderId);
 	void HandlePartyLeft(const FUniqueNetId& LocalUserId, const FOnlinePartyId& PartyId);
 	void HandlePartyMemberExited(const FUniqueNetId& LocalUserId, const FOnlinePartyId& PartyId, const FUniqueNetId& MemberId, EMemberExitedReason ExitReason);
 	void HandlePartyMemberDataReceived(const FUniqueNetId& LocalUserId, const FOnlinePartyId& PartyId, const FUniqueNetId& MemberId, const FName& Namespace, const FOnlinePartyData& PartyMemberData);
 	void HandlePartyMemberJoined(const FUniqueNetId& LocalUserId, const FOnlinePartyId& PartyId, const FUniqueNetId& MemberId);
-	UE_DEPRECATED(5.1, "Use the new join in progress flow with USocialParty::RequestJoinInProgress.")
-	void HandlePartyMemberJIP(const FUniqueNetId& LocalUserId, const FOnlinePartyId& PartyId, bool Success, int32 DeniedResultCode);
 	void HandlePartyMemberPromoted(const FUniqueNetId& LocalUserId, const FOnlinePartyId& PartyId, const FUniqueNetId& NewLeaderId);
 	void HandlePartyPromotionLockoutChanged(const FUniqueNetId& LocalUserId, const FOnlinePartyId& PartyId, bool bArePromotionsLocked);
 
diff --git a/Engine/Plugins/Online/OnlineFramework/Source/Party/Public/SocialToolkit.h b/Engine/Plugins/Online/OnlineFramework/Source/Party/Public/SocialToolkit.h
--- a/Engine/Plugins/Online/OnlineFramework/Source/Party/Public/SocialToolkit.h
+++ b/Engine/Plugins/Online/OnlineFramework/Source/Party/Public/SocialToolkit.h
@@ -186,10 +186,6 @@ protected:
 	virtual void OnBlockPlayerComplete(int32 LocalUserNum, bool bWasSuccessful, const FUniqueNetId& BlockedPlayerID, const FString& ListName, const FString& ErrorStr, ESocialSubsystem SubsystemType) {}
 	virtual void OnUnblockPlayerComplete(int32 LocalUserNum, bool bWasSuccessful, const FUniqueNetId& UnblockedPlayerID, const FString& ListName, const FString& ErrorStr, ESocialSubsystem SubsystemType) {}
 	
-public:
-	UE_DEPRECATED(5.1, "This function is deprecated and will be removed; please use USocialUser::HandlePartyRequestToJoinSent.")
-	virtual void OnRequestToJoinPartyComplete(const FUniqueNetId& PartyLeaderId, ERequestToJoinPartyCompletionResult Result) {}
-
 protected:
 	/** Called when a Friend's presence did change */
 	virtual void OnFriendPresenceDidChange(const USocialUser& FriendSocialUser, const TSharedRef<FOnlineUserPresence>& NewPresence, ESocialSubsystem SubsystemType) {}
diff --git a/Engine/Plugins/Online/OnlineFramework/Source/Party/Public/User/SocialUser.h b/Engine/Plugins/Online/OnlineFramework/Source/Party/Public/User/SocialUser.h
--- a/Engine/Plugins/Online/OnlineFramework/Source/Party/Public/User/SocialUser.h
+++ b/Engine/Plugins/Online/OnlineFramework/Source/Party/Public/User/SocialUser.h
@@ -201,9 +201,6 @@ protected:
 	virtual void HandleSetNicknameComplete(int32 LocalUserNum, const FUniqueNetId& FriendId, const FString& ListName, const FOnlineError& Error);
 	virtual void SetSubsystemId(ESocialSubsystem SubsystemType, const FUniqueNetIdRepl& SubsystemId);
 
-	UE_DEPRECATED(5.1, "This function is deprecated and will be removed.")
-	virtual void NotifyRequestToJoinSent(const FDateTime& ExpiresAt) {}
-
 	virtual void NotifyRequestToJoinReceived(const IOnlinePartyRequestToJoinInfo& Request) {}
 	virtual void NotifyRequestToJoinRemoved(const IOnlinePartyRequestToJoinInfo& Request, EPartyRequestToJoinRemovedReason Reason) {}
 	int32 NumPendingQueries = 0;
diff --git a/Engine/Plugins/Online/OnlineSubsystem/Source/Private/OnlinePartyInterface.cpp b/Engine/Plugins/Online/OnlineSubsystem/Source/Private/OnlinePartyInterface.cpp
--- a/Engine/Plugins/Online/OnlineSubsystem/Source/Private/OnlinePartyInterface.cpp
+++ b/Engine/Plugins/Online/OnlineSubsystem/Source/Private/OnlinePartyInterface.cpp
@@ -136,39 +136,6 @@ FString FPartyInvitationRecipient::ToDebugString() const
 	return FString::Printf(TEXT("Id=[%s], PlatformData=[%s]"), *Id->ToDebugString(), *PlatformData);
 }
 
-PRAGMA_DISABLE_DEPRECATION_WARNINGS
-FDelegateHandle IOnlinePartySystem::AddOnPartyInviteReceivedDelegate_Handle(const FOnPartyInviteReceivedDelegate& Delegate)
-PRAGMA_ENABLE_DEPRECATION_WARNINGS
-{
-	auto DeprecationHelperLambda = [Delegate](const FUniqueNetId& LocalUserId, const IOnlinePartyJoinInfo& Invitation)
-	{
-		Delegate.ExecuteIfBound(LocalUserId, *Invitation.GetPartyId(), *Invitation.GetSourceUserId());
-	};
-	return OnPartyInviteReceivedExDelegates.Add(FOnPartyInviteReceivedExDelegate::CreateLambda(DeprecationHelperLambda));
-}
-
-PRAGMA_DISABLE_DEPRECATION_WARNINGS
-FDelegateHandle IOnlinePartySystem::AddOnPartyInviteRemovedDelegate_Handle(const FOnPartyInviteRemovedDelegate& Delegate)
-PRAGMA_ENABLE_DEPRECATION_WARNINGS
-{
-	auto DeprecationHelperLambda = [Delegate](const FUniqueNetId& LocalUserId, const IOnlinePartyJoinInfo& Invitation, EPartyInvitationRemovedReason Reason)
-	{
-		Delegate.ExecuteIfBound(LocalUserId, *Invitation.GetPartyId(), *Invitation.GetSourceUserId(), Reason);
-	};
-	return OnPartyInviteRemovedExDelegates.Add(FOnPartyInviteRemovedExDelegate::CreateLambda(DeprecationHelperLambda));
-}
-
-PRAGMA_DISABLE_DEPRECATION_WARNINGS
-void IOnlinePartySystem::QueryPartyJoinability(const FUniqueNetId& LocalUserId, const IOnlinePartyJoinInfo& OnlinePartyJoinInfo, const FOnQueryPartyJoinabilityComplete& Delegate)
-PRAGMA_ENABLE_DEPRECATION_WARNINGS
-{
-	// Default implementation will call the ex version.
-	QueryPartyJoinability(LocalUserId, OnlinePartyJoinInfo, FOnQueryPartyJoinabilityCompleteEx::CreateLambda([Delegate](const FUniqueNetId& LambdaLocalUserId, const FOnlinePartyId& LambdaPartyId, const FQueryPartyJoinabilityResult& QueryPartyJoinabilityResult)
-	{
-		Delegate.ExecuteIfBound(LambdaLocalUserId, LambdaPartyId, QueryPartyJoinabilityResult.EnumResult, QueryPartyJoinabilityResult.SubCode);
-	}));
-}
-
 bool FPartyConfiguration::operator==(const FPartyConfiguration& Other) const
 {
 	return JoinRequestAction == Other.JoinRequestAction &&
diff --git a/Engine/Plugins/Online/OnlineSubsystem/Source/Public/Interfaces/OnlinePartyInterface.h b/Engine/Plugins/Online/OnlineSubsystem/Source/Public/Interfaces/OnlinePartyInterface.h
--- a/Engine/Plugins/Online/OnlineSubsystem/Source/Public/Interfaces/OnlinePartyInterface.h
+++ b/Engine/Plugins/Online/OnlineSubsystem/Source/Public/Interfaces/OnlinePartyInterface.h
@@ -850,16 +850,6 @@ DECLARE_DELEGATE_ThreeParams(FOnCreatePartyComplete, const FUniqueNetId& /*Local
  * @param NotApprovedReason - client defined value describing why you were not approved
  */
 DECLARE_DELEGATE_FourParams(FOnJoinPartyComplete, const FUniqueNetId& /*LocalUserId*/, const FOnlinePartyId& /*PartyId*/, const EJoinPartyCompletionResult /*Result*/, const int32 /*NotApprovedReason*/);
-/**
- * Party query joinability async task completed callback
- *
- * @param LocalUserId - id of user that initiated the request
- * @param PartyId - id associated with the party
- * @param Result - result of the operation
- * @param NotApprovedReason - client defined value describing why you were not approved
- */
-UE_DEPRECATED(5.1, "Use FOnQueryPartyJoinabilityCompleteEx")
-DECLARE_DELEGATE_FourParams(FOnQueryPartyJoinabilityComplete, const FUniqueNetId& /*LocalUserId*/, const FOnlinePartyId& /*PartyId*/, const EJoinPartyCompletionResult /*Result*/, const int32 /*NotApprovedReason*/);
 /**
  * Party query joinability async task completed callback
  *
@@ -989,20 +979,6 @@ PARTY_DECLARE_DELEGATETYPE(OnPartyExited);
 DECLARE_MULTICAST_DELEGATE_FourParams(F_PREFIX(OnPartyStateChanged), const FUniqueNetId& /*LocalUserId*/, const FOnlinePartyId& /*PartyId*/, EPartyState /*State*/, EPartyState /*PreviousState*/);
 PARTY_DECLARE_DELEGATETYPE(OnPartyStateChanged);
 
-/**
-* Notification when a player has been approved for JIP
-* @param LocalUserId - id associated with this notification
-* @param PartyId - id associated with the party
-* @param Success - whether the join in progress action succeeded
-* @param DeniedResultCode - descriptive reason for a failure to join
-*/
-DECLARE_MULTICAST_DELEGATE_ThreeParams(F_PREFIX(OnPartyJIP), const FUniqueNetId& /*LocalUserId*/, const FOnlinePartyId& /*PartyId*/, bool /*Success*/);
-UE_DEPRECATED(5.1, "Use the new join in progress flow with USocialParty::RequestJoinInProgress.")
-typedef FOnPartyJIP::FDelegate FOnPartyJIPDelegate;
-DECLARE_MULTICAST_DELEGATE_FourParams(F_PREFIX(OnPartyJIPResponse), const FUniqueNetId& /*LocalUserId*/, const FOnlinePartyId& /*PartyId*/, bool /*Success*/, int32 /*DeniedResultCode*/);
-UE_DEPRECATED(5.1, "Use the new join in progress flow with USocialParty::RequestJoinInProgress.")
-PARTY_DECLARE_DELEGATETYPE(OnPartyJIPResponse);
-
 /**
  * Notification when player promotion is locked out.
  * @param LocalUserId - id associated with this notification
@@ -1021,17 +997,6 @@ PARTY_DECLARE_DELEGATETYPE(OnPartyPromotionLockoutChanged);
 DECLARE_MULTICAST_DELEGATE_ThreeParams(F_PREFIX(OnPartyConfigChanged), const FUniqueNetId& /*LocalUserId*/, const FOnlinePartyId& /*PartyId*/, const FPartyConfiguration& /*PartyConfig*/);
 PARTY_DECLARE_DELEGATETYPE(OnPartyConfigChanged);
 
-/**
- * Notification when party data is updated
- * Deprecated - Use OnPartyDataReceived
- * @param LocalUserId - id associated with this notification
- * @param PartyId - id associated with the party
- * @param PartyData - party data that was updated
- */
-DECLARE_MULTICAST_DELEGATE_ThreeParams(F_PREFIX(OnPartyDataReceivedConst), const FUniqueNetId& /*LocalUserId*/, const FOnlinePartyId& /*PartyId*/, const FOnlinePartyData& /*PartyData*/);
-UE_DEPRECATED(5.1, "Use OnPartyDataReceived instead of OnPartyDataReceivedConst")
-typedef FOnPartyDataReceivedConst::FDelegate FOnPartyDataReceivedConstDelegate;
-
 /**
  * Notification when party data is updated
  * @param LocalUserId - id associated with this notification
@@ -1070,18 +1035,6 @@ PARTY_DECLARE_DELEGATETYPE(OnPartyMemberExited);
 DECLARE_MULTICAST_DELEGATE_ThreeParams(F_PREFIX(OnPartyMemberJoined), const FUniqueNetId& /*LocalUserId*/, const FOnlinePartyId& /*PartyId*/, const FUniqueNetId& /*MemberId*/);
 PARTY_DECLARE_DELEGATETYPE(OnPartyMemberJoined);
 
-/**
- * Notification when party member data is updated
- * Deprecated - Use OnPartyMemberDataReceived
- * @param LocalUserId - id associated with this notification
- * @param PartyId - id associated with the party
- * @param MemberId - id of member that had updated data
- * @param PartyMemberData - party member data that was updated
- */
-DECLARE_MULTICAST_DELEGATE_FourParams(F_PREFIX(OnPartyMemberDataReceivedConst), const FUniqueNetId& /*LocalUserId*/, const FOnlinePartyId& /*PartyId*/, const FUniqueNetId& /*MemberId*/, const FOnlinePartyData& /*PartyMemberData*/);
-UE_DEPRECATED(5.1, "Use OnPartyMemberDataReceived instead of OnPartyMemberDataReceivedConst")
-typedef FOnPartyMemberDataReceivedConst::FDelegate FOnPartyMemberDataReceivedConstDelegate;
-
 /**
  * Notification when party member data is updated
  * @param LocalUserId - id associated with this notification
@@ -1110,17 +1063,6 @@ PARTY_DECLARE_DELEGATETYPE(OnPartyInvitesChanged);
 DECLARE_MULTICAST_DELEGATE_FourParams(F_PREFIX(OnPartyInviteRequestReceived), const FUniqueNetId& /*LocalUserId*/, const FOnlinePartyId& /*PartyId*/, const FUniqueNetId& /*SenderId*/, const FUniqueNetId& /*RequestForId*/);
 PARTY_DECLARE_DELEGATETYPE(OnPartyInviteRequestReceived);
 
-/**
- * Notification when a new invite is received
- * Deprecated - Use OnPartyInviteReceivedEx
- * @param LocalUserId id associated with this notification
- * @param PartyId id associated with the party
- * @param SenderId id of member that sent the invite
- */
-DECLARE_MULTICAST_DELEGATE_ThreeParams(F_PREFIX(OnPartyInviteReceived), const FUniqueNetId& /*LocalUserId*/, const FOnlinePartyId& /*PartyId*/, const FUniqueNetId& /*SenderId*/);
-UE_DEPRECATED(5.1, "Use OnPartyInviteReceivedEx instead of OnPartyInviteReceived")
-typedef FOnPartyInviteReceived::FDelegate FOnPartyInviteReceivedDelegate;
-
 /**
  * Notification when a new invite is received
  * @param LocalUserId id associated with this notification
@@ -1129,18 +1071,6 @@ typedef FOnPartyInviteReceived::FDelegate FOnPartyInviteReceivedDelegate;
 DECLARE_MULTICAST_DELEGATE_TwoParams(F_PREFIX(OnPartyInviteReceivedEx), const FUniqueNetId& /*LocalUserId*/, const IOnlinePartyJoinInfo& /*Invitation*/);
 PARTY_DECLARE_DELEGATETYPE(OnPartyInviteReceivedEx);
 
-/**
- * Notification when an invite has been removed
- * Deprecated - Use OnPartyInviteRemovedEx
- * @param LocalUserId id associated with this notification
- * @param PartyId id associated with the party
- * @param SenderId id of member that sent the invite
- * @param Reason reason the invite has been removed
- */
-DECLARE_MULTICAST_DELEGATE_FourParams(F_PREFIX(OnPartyInviteRemoved), const FUniqueNetId& /*LocalUserId*/, const FOnlinePartyId& /*PartyId*/, const FUniqueNetId& /*SenderId*/, EPartyInvitationRemovedReason /*Reason*/);
-UE_DEPRECATED(5.1, "Use OnPartyInviteRemovedEx instead of OnPartyInviteRemoved")
-typedef FOnPartyInviteRemoved::FDelegate FOnPartyInviteRemovedDelegate;
-
 /**
  * Notification when an invite has been removed
  * @param LocalUserId id associated with this notification
@@ -1169,18 +1099,6 @@ PARTY_DECLARE_DELEGATETYPE(OnPartyInviteResponseReceived);
 DECLARE_MULTICAST_DELEGATE_ThreeParams(F_PREFIX(OnPartyJoinRequestReceived), const FUniqueNetId& /*LocalUserId*/, const FOnlinePartyId& /*PartyId*/, const IOnlinePartyPendingJoinRequestInfo& /*JoinRequestInfo*/);
 PARTY_DECLARE_DELEGATETYPE(OnPartyJoinRequestReceived);
 
-/**
- * Notification when a new reservation request is received
- * @param LocalUserId - id associated with this notification
- * @param PartyId - id associated with the party
- * @param SenderId - id of member that sent the request
- * @param Platform - platform of member that sent the request
- * @param PartyData - data provided by the sender for the leader to use to determine joinability
- */
-DECLARE_MULTICAST_DELEGATE_ThreeParams(F_PREFIX(OnPartyJIPRequestReceived), const FUniqueNetId& /*LocalUserId*/, const FOnlinePartyId& /*PartyId*/, const FUniqueNetId& /*SenderId*/);
-UE_DEPRECATED(5.1, "Use the new join in progress flow with USocialParty::RequestJoinInProgress.")
-PARTY_DECLARE_DELEGATETYPE(OnPartyJIPRequestReceived);
-
 /**
  * Notification when a player wants to know if the party is in a joinable state
  * @param LocalUserId - id associated with this notification
@@ -1236,18 +1154,6 @@ PARTY_DECLARE_DELEGATETYPE(OnPartyRequestToJoinReceived);
 DECLARE_MULTICAST_DELEGATE_FiveParams(F_PREFIX(OnPartyRequestToJoinRemoved), const FUniqueNetId& /*LocalUserId*/, const FOnlinePartyId& /*PartyId*/, const FUniqueNetId& /*RequesterId*/, const IOnlinePartyRequestToJoinInfo& /*Request*/, EPartyRequestToJoinRemovedReason /*Reason*/);
 PARTY_DECLARE_DELEGATETYPE(OnPartyRequestToJoinRemoved);
 
-#define DEFINE_DEPRECATED_PARTY_DELEGATE(DeprecatedDelegateName, DelegateName) \
-	virtual void Clear##DeprecatedDelegateName##Delegate_Handle(FDelegateHandle& Handle) \
-	{ \
-		DelegateName##Delegates.Remove(Handle); \
-		Handle.Reset(); \
-	} \
-	virtual void Clear##DeprecatedDelegateName##Delegates(void* Object) \
-	{ \
-		DelegateName##Delegates.RemoveAll(Object); \
-	}
-
-
 /**
  * Interface definition for the online party services 
  * Allows for forming a party and communicating with party members
@@ -1338,32 +1244,6 @@ public:
 	*/
 	virtual void ClearRequestToJoinParty(const FUniqueNetId& LocalUserId, const FOnlinePartyId& PartyId, const FUniqueNetId& Sender, EPartyRequestToJoinRemovedReason Reason) = 0;
 
-	/**
-	* Join an existing game session from within a party
-	*
-	* @param LocalUserId - user making the request
-	* @param OnlinePartyJoinInfo - join information containing data such as party id, leader id
-	* @param Delegate - called on completion
-	*
-	* @return true if task was started
-	*/
-	UE_DEPRECATED(5.1, "Use the new join in progress flow with USocialParty::RequestJoinInProgress.")
-	virtual bool JIPFromWithinParty(const FUniqueNetId& LocalUserId, const FOnlinePartyId& PartyId, const FUniqueNetId& PartyLeaderId) { return false; }
-
-	/**
-	 * Query a party to check it's current joinability
-	 * Intended to be used before a call to LeaveParty (to leave your existing party, which would then be followed by JoinParty)
-	 * Note that the party's joinability can change from moment to moment so a successful response for this does not guarantee a successful JoinParty
-	 *
-	 * @param LocalUserId - user making the request
-	 * @param OnlinePartyJoinInfo - join information containing data such as party id, leader id
-	 * @param Delegate - called on completion
-	 */
-	PRAGMA_DISABLE_DEPRECATION_WARNINGS // Suppress deprecation warning for FOnQueryPartyJoinabilityComplete
-	UE_DEPRECATED(5.1, "Use QueryPartyJoinability with the FOnQueryPartyJoinabilityCompleteEx delegate instead")
-	virtual void QueryPartyJoinability(const FUniqueNetId& LocalUserId, const IOnlinePartyJoinInfo& OnlinePartyJoinInfo, const FOnQueryPartyJoinabilityComplete& Delegate);
-	PRAGMA_ENABLE_DEPRECATION_WARNINGS
-
 	/**
 	 * Query a party to check it's current joinability
 	 * Intended to be used before a call to LeaveParty (to leave your existing party, which would then be followed by JoinParty)
@@ -1426,20 +1306,6 @@ public:
 	*/
 	virtual bool ApproveJoinRequest(const FUniqueNetId& LocalUserId, const FOnlinePartyId& PartyId, const FUniqueNetId& RecipientId, bool bIsApproved, int32 DeniedResultCode = 0) = 0;
 
-	/**
-	* Approve a request to join the JIP match a party is in.
-	*
-	* @param LocalUserId - user making the request
-	* @param PartyId - id of an existing party
-	* @param RecipientId - id of the user being invited
-	* @param bIsApproved - whether the join request was approved or not
-	* @param DeniedResultCode - used when bIsApproved is false - client defined value to return when leader denies approval
-	*
-	* @return true if task was started
-	*/
-	UE_DEPRECATED(5.1, "Use the new join in progress flow with USocialParty::RequestJoinInProgress.")
-	virtual bool ApproveJIPRequest(const FUniqueNetId& LocalUserId, const FOnlinePartyId& PartyId, const FUniqueNetId& RecipientId, bool bIsApproved, int32 DeniedResultCode = 0) { return false; }
-
 	/**
 	 * Respond to a query joinability request.  This reflects the current party's joinability state and can change from moment to moment, and therefore does not guarantee a successful join.
 	 *
@@ -1813,17 +1679,6 @@ public:
 	*/
 	DEFINE_ONLINE_DELEGATE_FOUR_PARAM(OnPartyStateChanged, const FUniqueNetId& /*LocalUserId*/, const FOnlinePartyId& /*PartyId*/, EPartyState /*State*/, EPartyState /*PreviousState*/);
 
-	/**
-	* notification of when a player had been approved to Join In Progress
-	* @param LocalUserId - id associated with this notification
-	* @param PartyId - id associated with the party
-	* @param Success - whether the join in progress action succeeded
-	* @param DeniedResultCode - descriptive reason for a failure to join
-	*/
-	PRAGMA_DISABLE_DEPRECATION_WARNINGS
-	DEFINE_ONLINE_DELEGATE_FOUR_PARAM(OnPartyJIPResponse, const FUniqueNetId& /*LocalUserId*/, const FOnlinePartyId& /*PartyId*/, bool /*Success*/, int32 /*DeniedResultCode*/);
-	PRAGMA_ENABLE_DEPRECATION_WARNINGS
-
 	/**
 	 * Notification when player promotion is locked out.
 	 *
@@ -1905,11 +1760,6 @@ public:
 	 * @param Invitation - the invitation that was received
 	 */
 	DEFINE_ONLINE_DELEGATE_TWO_PARAM(OnPartyInviteReceivedEx, const FUniqueNetId& /*LocalUserId*/, const IOnlinePartyJoinInfo& /*Invitation*/);
-	PRAGMA_DISABLE_DEPRECATION_WARNINGS // Suppress deprecation warning for FOnPartyInviteReceivedDelegate
-	UE_DEPRECATED(5.1, "Use OnPartyInviteReceivedEx instead of OnPartyInviteReceived")
-	virtual FDelegateHandle AddOnPartyInviteReceivedDelegate_Handle(const FOnPartyInviteReceivedDelegate& Delegate);
-	DEFINE_DEPRECATED_PARTY_DELEGATE(OnPartyInviteReceived, OnPartyInviteReceivedEx);
-	PRAGMA_ENABLE_DEPRECATION_WARNINGS
 
 	/**
 	 * Notification when an invite has been removed
@@ -1918,11 +1768,6 @@ public:
 	 * @param Reason reason the invitation was removed
 	 */
 	DEFINE_ONLINE_DELEGATE_THREE_PARAM(OnPartyInviteRemovedEx, const FUniqueNetId& /*LocalUserId*/, const IOnlinePartyJoinInfo& /*Invitation*/, EPartyInvitationRemovedReason /*Reason*/);
-	PRAGMA_DISABLE_DEPRECATION_WARNINGS // Suppress deprecation warning for FOnPartyInviteRemovedDelegate
-	UE_DEPRECATED(5.1, "Use OnPartyInviteRemovedEx instead of OnPartyInviteRemoved")
-	virtual FDelegateHandle AddOnPartyInviteRemovedDelegate_Handle(const FOnPartyInviteRemovedDelegate& Delegate);
-	DEFINE_DEPRECATED_PARTY_DELEGATE(OnPartyInviteRemoved, OnPartyInviteRemovedEx);
-	PRAGMA_ENABLE_DEPRECATION_WARNINGS
 
 	/**
 	 * Notification when an invitation response is received
@@ -1942,19 +1787,6 @@ public:
 	 */
 	DEFINE_ONLINE_DELEGATE_THREE_PARAM(OnPartyJoinRequestReceived, const FUniqueNetId& /*LocalUserId*/, const FOnlinePartyId& /*PartyId*/, const IOnlinePartyPendingJoinRequestInfo& /*JoinRequestInfo*/);
 
-	/**
-	* Notification when a new reservation request is received
-	* Subscriber is expected to call ApproveJoinRequest
-	* @param LocalUserId - id associated with this notification
-	* @param PartyId - id associated with the party
-	* @param SenderId - id of member that sent the request
-	* @param Platform - platform of member that sent the request
-	* @param PartyData - data provided by the sender for the leader to use to determine joinability
-	*/
-	PRAGMA_DISABLE_DEPRECATION_WARNINGS
-	DEFINE_ONLINE_DELEGATE_THREE_PARAM(OnPartyJIPRequestReceived, const FUniqueNetId& /*LocalUserId*/, const FOnlinePartyId& /*PartyId*/, const FUniqueNetId& /*SenderId*/);
-	PRAGMA_ENABLE_DEPRECATION_WARNINGS
-
 	/**
 	 * Notification when a player wants to know if the party is in a joinable state
 	 * Subscriber is expected to call RespondToQueryJoinability
