Cleanup deprecated online functionality

diff --git a/Engine/Plugins/Online/OnlineSubsystem/Source/Public/Interfaces/OnlineEventsInterface.h b/Engine/Plugins/Online/OnlineSubsystem/Source/Public/Interfaces/OnlineEventsInterface.h
--- a/Engine/Plugins/Online/OnlineSubsystem/Source/Public/Interfaces/OnlineEventsInterface.h
+++ b/Engine/Plugins/Online/OnlineSubsystem/Source/Public/Interfaces/OnlineEventsInterface.h
@@ -16,60 +16,8 @@ ONLINESUBSYSTEM_API DECLARE_LOG_CATEGORY_EXTERN(LogOnlineEvents, Log, All);
 }
 
 class FUniqueNetId;
-// Temporary class to assist in deprecation of changing the key from FName to FString. This will eventually be a typedef to FOnlineKeyValuePairs< FString, FVariantData >
-class FOnlineEventParms : public FOnlineKeyValuePairs< FString, FVariantData >
-{
-	typedef FOnlineKeyValuePairs<FString, FVariantData> Super;
-
-public:
-	FORCEINLINE FOnlineEventParms() {}
-	FORCEINLINE FOnlineEventParms(FOnlineEventParms&& Other) : Super(MoveTemp(Other)) {}
-	FORCEINLINE FOnlineEventParms(const FOnlineEventParms& Other) : Super(Other) {}
-	FORCEINLINE FOnlineEventParms& operator=(FOnlineEventParms&& Other) { Super::operator=(MoveTemp(Other)); return *this; }
-	FORCEINLINE FOnlineEventParms& operator=(const FOnlineEventParms& Other) { Super::operator=(Other); return *this; }
-
-public:
-	UE_DEPRECATED(4.26, "FOnlineEventParms now uses FString for the key type instead of FName")
-		FOnlineEventParms(const FOnlineKeyValuePairs< FName, FVariantData >& DeprecatedValues)
-	{
-		Reserve(DeprecatedValues.Num());
-		for (const TPair< FName, FVariantData >& DeprecatedValue : DeprecatedValues)
-		{
-			Super::Emplace(DeprecatedValue.Key.ToString(), DeprecatedValue.Value);
-		}
-	}
-
-	UE_DEPRECATED(4.26, "FOnlineEventParms now uses FString for the key type instead of FName")
-		FOnlineEventParms(FOnlineKeyValuePairs< FName, FVariantData >&& DeprecatedValues)
-	{
-		Reserve(DeprecatedValues.Num());
-		for (TPair< FName, FVariantData >& DeprecatedValue : DeprecatedValues)
-		{
-			Super::Emplace(DeprecatedValue.Key.ToString(), MoveTemp(DeprecatedValue.Value));
-		}
-	}
- 	using Super::Add;
- 	UE_DEPRECATED(4.26, "FOnlineEventParms now uses FString for the key type instead of FName")
- 		FVariantData& Add(const FName& InKey, const FVariantData& InValue)
- 	{
- 		return Super::Add(InKey.ToString(), InValue);
- 	}
 
-	using Super::Emplace;
-	UE_DEPRECATED(4.26, "FOnlineEventParms now uses FString for the key type instead of FName")
-		FVariantData& Emplace(const FName& InKey, const FVariantData& InValue)
-	{
-		return Super::Emplace(InKey.ToString(), InValue);
-	}
-
-	using Super::Find;
-	UE_DEPRECATED(4.26, "FOnlineEventParms now uses FString for the key type instead of FName")
-		FVariantData* Find(const FName& InKey)
-	{
-		return Super::Find(InKey.ToString());
-	}
-};
-// typedef FOnlineKeyValuePairs< FString, FVariantData > FOnlineEventParms;
+class FOnlineEventParms : public FOnlineKeyValuePairs<FString, FVariantData> {};
 
 /**
  *	IOnlineEvents - Interface class for events
diff --git a/Engine/Plugins/Online/OnlineSubsystem/Source/Public/Interfaces/OnlineGameMatchesInterface.h b/Engine/Plugins/Online/OnlineSubsystem/Source/Public/Interfaces/OnlineGameMatchesInterface.h
--- a/Engine/Plugins/Online/OnlineSubsystem/Source/Public/Interfaces/OnlineGameMatchesInterface.h
+++ b/Engine/Plugins/Online/OnlineSubsystem/Source/Public/Interfaces/OnlineGameMatchesInterface.h
@@ -392,20 +392,6 @@ public:
 	 * @param CompletionDelegate - Completion delegate called when ReportGameMatchResults is complete
 	 */
 	virtual void ReportGameMatchResults(const FUniqueNetId& UserId, const FString& MatchId, const FFinalGameMatchReport& FinalReport, const FOnGameMatchReportComplete& CompletionDelegate) = 0;
-	
-	/**
-	 * Provides a way to leave feedback on a game match
-	 *
-	 * This call will invoke a UI for the player to provide feedback about the recent match they participated in
-	 *
-	 * @param UserId - Id of the user leaving feedback on the match
-	 * @param MatchId - Id of the match to leave feedback on
-	 * @param bReviewTeam - Whether not to review a team or only the members of the team
-	 * @param CompletionDelegate - Completion delegate called when GameMatchFeedback is complete
-	 */
-	UE_DEPRECATED(5.0, "ProvideGameMatchFeedback is now deprecated and will be removed.")
-	virtual void ProvideGameMatchFeedback(const FUniqueNetId& UserId, const FString& MatchId, const bool bReviewTeam, const FOnGameMatchFeedbackComplete& CompletionDelegate) {}
-
 };
 
 #if UE_ENABLE_INCLUDE_ORDER_DEPRECATED_IN_5_2
diff --git a/Engine/Plugins/Online/OnlineSubsystem/Source/Public/OnlineSubsystemTypes.h b/Engine/Plugins/Online/OnlineSubsystem/Source/Public/OnlineSubsystemTypes.h
--- a/Engine/Plugins/Online/OnlineSubsystem/Source/Public/OnlineSubsystemTypes.h
+++ b/Engine/Plugins/Online/OnlineSubsystem/Source/Public/OnlineSubsystemTypes.h
@@ -798,27 +798,9 @@ public:
 
 	FName Type = NAME_Unset;
 	
-	UE_DEPRECATED(5.0, "This FUniqueNetIdString Create method is deprecated. Please use Create(const FString& InUniqueNetId, const FName InType)")
-	static FUniqueNetIdStringRef Create(const FString& InUniqueNetId)
-	{
-		PRAGMA_DISABLE_DEPRECATION_WARNINGS
-		return MakeShareable(new FUniqueNetIdString(InUniqueNetId, NAME_Unset));
-		PRAGMA_ENABLE_DEPRECATION_WARNINGS
-	}
-
-	UE_DEPRECATED(5.0, "This FUniqueNetIdString Create method is deprecated. Please use Create(FString&& InUniqueNetId, const FName InType)")
-	static FUniqueNetIdStringRef Create(FString&& InUniqueNetId)
-	{
-		PRAGMA_DISABLE_DEPRECATION_WARNINGS
-		return MakeShareable(new FUniqueNetIdString(MoveTemp(InUniqueNetId), NAME_Unset));
-		PRAGMA_ENABLE_DEPRECATION_WARNINGS
-	}
-
 	static FUniqueNetIdStringRef Create(const FString& InUniqueNetId, const FName InType)
 	{
-		PRAGMA_DISABLE_DEPRECATION_WARNINGS
 		return MakeShareable(new FUniqueNetIdString(InUniqueNetId, InType));
-		PRAGMA_ENABLE_DEPRECATION_WARNINGS
 	}
 
 	static FUniqueNetIdStringRef Create(FString&& InUniqueNetId, const FName InType)
@@ -878,8 +860,7 @@ public:
 		return GetTypeHashHelper(UniqueNetIdStr);
 	}
 
-public:
-	UE_DEPRECATED(5.0, "Public constructors of FUniqueNetId types are deprecated. Please use the corresponding Create method instead to create a FUniqueNetIdRef")
+protected:
 	FUniqueNetIdString() = default;
 
 	/**
@@ -887,7 +868,6 @@ public:
 	 *
 	 * @param InUniqueNetId the id to set ours to
 	 */
-	UE_DEPRECATED(5.0, "Public constructors of FUniqueNetId types are deprecated. Please use the corresponding Create method instead to create a FUniqueNetIdRef")
 	explicit FUniqueNetIdString(const FString& InUniqueNetId)
 		: UniqueNetIdStr(InUniqueNetId)
 		, Type(NAME_Unset)
@@ -899,7 +879,6 @@ public:
 	 *
 	 * @param InUniqueNetId the id to set ours to
 	 */
-	UE_DEPRECATED(5.0, "Public constructors of FUniqueNetId types are deprecated. Please use the corresponding Create method instead to create a FUniqueNetIdRef")
 	explicit FUniqueNetIdString(FString&& InUniqueNetId)
 		: UniqueNetIdStr(MoveTemp(InUniqueNetId))
 		, Type(NAME_Unset)
@@ -911,7 +890,6 @@ public:
 	 *
 	 * @param Src the id to copy
 	 */
-	UE_DEPRECATED(5.0, "Public constructors of FUniqueNetId types are deprecated. Please use the corresponding Create method instead to create a FUniqueNetIdRef")
 	explicit FUniqueNetIdString(const FUniqueNetId& Src)
 		: UniqueNetIdStr(Src.ToString())
 		, Type(Src.GetType())
@@ -921,14 +899,12 @@ public:
 	/**
 	* don.eubanks - Including a constructor that allows for type passing to make transitioning easier, if we determine we want to abstract-ify this class, this constructor will be removed
 	*/
-	UE_DEPRECATED(5.0, "Public constructors of FUniqueNetId types are deprecated. Please use the corresponding Create method instead to create a FUniqueNetIdRef")
 	FUniqueNetIdString(const FString& InUniqueNetId, const FName InType)
 		: UniqueNetIdStr(InUniqueNetId)
 		, Type(InType)
 	{
 	}
 
-protected:
 	FUniqueNetIdString(FString&& InUniqueNetId, const FName InType)
 		: UniqueNetIdStr(MoveTemp(InUniqueNetId))
 		, Type(InType)
@@ -946,9 +922,7 @@ public: \
 	template<typename... TArgs> \
 	static SUBCLASSNAME##Ref Create(TArgs&&... Args) \
 	{ \
-		PRAGMA_DISABLE_DEPRECATION_WARNINGS \
 		return MakeShareable(new SUBCLASSNAME(Forward<TArgs>(Args)...)); \
-		PRAGMA_ENABLE_DEPRECATION_WARNINGS \
 	} \
 	static SUBCLASSNAME##Ref Cast(const FUniqueNetIdRef& InNetId) \
 	{ \
@@ -982,24 +956,20 @@ public: \
 		static const SUBCLASSNAME##Ref EmptyId(Create()); \
 		return EmptyId; \
 	} \
-public: \
-	UE_DEPRECATED(5.0, "Public constructors of FUniqueNetId types are deprecated. Please use the ::Create(Args) method instead to create a FUniqueNetIdRef") \
+protected: \
 	SUBCLASSNAME() \
 		: FUniqueNetIdString() \
 	{ \
 		Type = TYPE; \
 	} \
-	UE_DEPRECATED(5.0, "Public constructors of FUniqueNetId types are deprecated. Please use the ::Create(Args) method instead to create a FUniqueNetIdRef") \
 	explicit SUBCLASSNAME(const FString& InUniqueNetId) \
 		: FUniqueNetIdString(InUniqueNetId, TYPE) \
 	{ \
 	} \
-	UE_DEPRECATED(5.0, "Public constructors of FUniqueNetId types are deprecated. Please use the ::Create(Args) method instead to create a FUniqueNetIdRef") \
 	explicit SUBCLASSNAME(FString&& InUniqueNetId) \
 		: FUniqueNetIdString(MoveTemp(InUniqueNetId), TYPE) \
 	{ \
 	} \
-	UE_DEPRECATED(5.0, "Public constructors of FUniqueNetId types are deprecated. Please use the ::Create(Args) method instead to create a FUniqueNetIdRef") \
 	explicit SUBCLASSNAME(const FUniqueNetId& Src) \
 		: FUniqueNetIdString(Src) \
 	{ \
diff --git a/Engine/Plugins/Online/OnlineSubsystemEOS/Source/OnlineSubsystemEOSPlus/Private/OnlineUserEOSPlus.cpp b/Engine/Plugins/Online/OnlineSubsystemEOS/Source/OnlineSubsystemEOSPlus/Private/OnlineUserEOSPlus.cpp
--- a/Engine/Plugins/Online/OnlineSubsystemEOS/Source/OnlineSubsystemEOSPlus/Private/OnlineUserEOSPlus.cpp
+++ b/Engine/Plugins/Online/OnlineSubsystemEOS/Source/OnlineSubsystemEOSPlus/Private/OnlineUserEOSPlus.cpp
@@ -67,9 +67,7 @@ inline FString BuildEOSPlusStringId(FUniqueNetIdPtr InBaseUniqueNetId, FUniqueNe
 }
 
 FUniqueNetIdEOSPlus::FUniqueNetIdEOSPlus(FUniqueNetIdPtr InBaseUniqueNetId, FUniqueNetIdPtr InEOSUniqueNetId)
-	PRAGMA_DISABLE_DEPRECATION_WARNINGS
 	: FUniqueNetIdString(BuildEOSPlusStringId(InBaseUniqueNetId, InEOSUniqueNetId), FName("EOSPlus"))
-	PRAGMA_ENABLE_DEPRECATION_WARNINGS
 	, BaseUniqueNetId(InBaseUniqueNetId)
 	, EOSUniqueNetId(InEOSUniqueNetId)
 {
diff --git a/Engine/Plugins/Online/OnlineSubsystemEOS/Source/OnlineSubsystemEOSPlus/Private/OnlineUserEOSPlus.h b/Engine/Plugins/Online/OnlineSubsystemEOS/Source/OnlineSubsystemEOSPlus/Private/OnlineUserEOSPlus.h
--- a/Engine/Plugins/Online/OnlineSubsystemEOS/Source/OnlineSubsystemEOSPlus/Private/OnlineUserEOSPlus.h
+++ b/Engine/Plugins/Online/OnlineSubsystemEOS/Source/OnlineSubsystemEOSPlus/Private/OnlineUserEOSPlus.h
@@ -55,7 +55,7 @@ public:
 	/** global static instance of invalid (zero) id */
 	static const FUniqueNetIdEOSPlusRef& EmptyId()
 	{
-		static const FUniqueNetIdEOSPlusRef EmptyId(Create());
+		static const FUniqueNetIdEOSPlusRef EmptyId(Create(nullptr, nullptr));
 		return EmptyId;
 	}
 
@@ -65,12 +65,6 @@ PACKAGE_SCOPE:
 	TArray<uint8> RawBytes;
 
 private:
-	PRAGMA_DISABLE_DEPRECATION_WARNINGS
-	FUniqueNetIdEOSPlus()
-	{
-	}
-	PRAGMA_ENABLE_DEPRECATION_WARNINGS
-
 	explicit FUniqueNetIdEOSPlus(FUniqueNetIdPtr InBaseUniqueNetId, FUniqueNetIdPtr InEOSUniqueNetId);
 };
 
diff --git a/Engine/Plugins/Online/OnlineSubsystemUtils/Source/OnlineSubsystemUtils/Private/OnlineBeaconClient.cpp b/Engine/Plugins/Online/OnlineSubsystemUtils/Source/OnlineSubsystemUtils/Private/OnlineBeaconClient.cpp
--- a/Engine/Plugins/Online/OnlineSubsystemUtils/Source/OnlineSubsystemUtils/Private/OnlineBeaconClient.cpp
+++ b/Engine/Plugins/Online/OnlineSubsystemUtils/Source/OnlineSubsystemUtils/Private/OnlineBeaconClient.cpp
@@ -220,20 +220,6 @@ void AOnlineBeaconClient::Tick(float DeltaTime)
 	}
 }
 
-void AOnlineBeaconClient::SetEncryptionToken(const FString& InEncryptionToken)
-{
-	EncryptionData.Identifier = InEncryptionToken;
-}
-
-void AOnlineBeaconClient::SetEncryptionKey(TArrayView<uint8> InEncryptionKey)
-{
-	if (CVarNetAllowEncryption.GetValueOnGameThread() != 0)
-	{
-		EncryptionData.Key.Reset(InEncryptionKey.Num());
-		EncryptionData.Key.Append(InEncryptionKey.GetData(), InEncryptionKey.Num());
-	}
-}
-
 void AOnlineBeaconClient::SetEncryptionData(const FEncryptionData& InEncryptionData)
 {
 	if (CVarNetAllowEncryption.GetValueOnGameThread() != 0)
diff --git a/Engine/Plugins/Online/OnlineSubsystemUtils/Source/OnlineSubsystemUtils/Private/OnlineBeaconHost.cpp b/Engine/Plugins/Online/OnlineSubsystemUtils/Source/OnlineSubsystemUtils/Private/OnlineBeaconHost.cpp
--- a/Engine/Plugins/Online/OnlineSubsystemUtils/Source/OnlineSubsystemUtils/Private/OnlineBeaconHost.cpp
+++ b/Engine/Plugins/Online/OnlineSubsystemUtils/Source/OnlineSubsystemUtils/Private/OnlineBeaconHost.cpp
@@ -556,20 +556,6 @@ void AOnlineBeaconHost::SendFailurePacket(UNetConnection* Connection, FNetCloseR
 	}
 }
 
-void AOnlineBeaconHost::SendFailurePacket(UNetConnection* Connection, const FText& ErrorText)
-{
-	if (Connection != nullptr)
-	{
-		FString ErrorMsg = ErrorText.ToString();
-
-		UE_LOG(LogBeacon, Log, TEXT("%s: Send failure: %s"), ToCStr(GetDebugName(Connection)), ToCStr(ErrorMsg));
-
-		FNetControlMessage<NMT_Failure>::Send(Connection, ErrorMsg);
-		Connection->FlushNet(true);
-		Connection->Close();
-	}
-}
-
 void AOnlineBeaconHost::CloseHandshakeConnection(UNetConnection* Connection)
 {
 	if (Connection && Connection->GetConnectionState() != USOCK_Closed)
diff --git a/Engine/Plugins/Online/OnlineSubsystemUtils/Source/OnlineSubsystemUtils/Private/OnlineSubsystemUtils.cpp b/Engine/Plugins/Online/OnlineSubsystemUtils/Source/OnlineSubsystemUtils/Private/OnlineSubsystemUtils.cpp
--- a/Engine/Plugins/Online/OnlineSubsystemUtils/Source/OnlineSubsystemUtils/Private/OnlineSubsystemUtils.cpp
+++ b/Engine/Plugins/Online/OnlineSubsystemUtils/Source/OnlineSubsystemUtils/Private/OnlineSubsystemUtils.cpp
@@ -152,33 +152,6 @@ UAudioComponent* CreateVoiceAudioComponent(uint32 SampleRate, int32 NumChannels)
 	return AudioComponent;
 }
 
-UVoipListenerSynthComponent* CreateVoiceSynthComponent(uint32 SampleRate)
-{
-	UVoipListenerSynthComponent* SynthComponentPtr = nullptr;
-	if (FAudioDeviceHandle AudioDevice = GEngine->GetMainAudioDevice())
-	{
-		SynthComponentPtr = NewObject<UVoipListenerSynthComponent>();
-		if (SynthComponentPtr)
-		{
-			const FSoftObjectPath VoiPSoundClassName = GetDefault<UAudioSettings>()->VoiPSoundClass;
-			if (VoiPSoundClassName.IsValid())
-			{
-				SynthComponentPtr->SoundClass = LoadObject<USoundClass>(nullptr, *VoiPSoundClassName.ToString());
-			}
-
-			SynthComponentPtr->bAlwaysPlay = CvarAlwaysPlayVoipComponent;
-			SynthComponentPtr->Initialize(SampleRate);
-		}
-		else
-		{
-			UE_LOG(LogVoiceDecode, Warning, TEXT("Unable to create voice synth component!"));
-		}
-	}
-
-	return SynthComponentPtr;
-}
-
-
 UVoipListenerSynthComponent* CreateVoiceSynthComponent(UWorld* World, uint32 SampleRate)
 {
 	UVoipListenerSynthComponent* SynthComponentPtr = nullptr;
diff --git a/Engine/Plugins/Online/OnlineSubsystemUtils/Source/OnlineSubsystemUtils/Public/OnlineBeaconClient.h b/Engine/Plugins/Online/OnlineSubsystemUtils/Source/OnlineSubsystemUtils/Public/OnlineBeaconClient.h
--- a/Engine/Plugins/Online/OnlineSubsystemUtils/Source/OnlineSubsystemUtils/Public/OnlineBeaconClient.h
+++ b/Engine/Plugins/Online/OnlineSubsystemUtils/Source/OnlineSubsystemUtils/Public/OnlineBeaconClient.h
@@ -70,22 +70,6 @@ class ONLINESUBSYSTEMUTILS_API AOnlineBeaconClient : public AOnlineBeacon
 	 */
 	bool InitClient(FURL& URL);
 
-	/**
-	 * Sets the encryption token that will be sent to servers on connection requests as a parameter to the NMT_Hello message.
-	 *
-	 * @param EncryptionToken the token to use
-	 */
-	UE_DEPRECATED(4.24, "Use SetEncryptionData instead")
-	void SetEncryptionToken(const FString& InEncryptionToken);
-
-	/**
-	 * Sets the encryption key that will be used for server connections.
-	 *
-	 * @param EncryptionKey the key to use
-	 */
-	UE_DEPRECATED(4.24, "Use SetEncryptionData instead")
-	void SetEncryptionKey(TArrayView<uint8> InEncryptionKey);
-
 	/**
 	 * Sets the encryption data that will be used for server connections.
 	 *
diff --git a/Engine/Plugins/Online/OnlineSubsystemUtils/Source/OnlineSubsystemUtils/Public/OnlineBeaconHost.h b/Engine/Plugins/Online/OnlineSubsystemUtils/Source/OnlineSubsystemUtils/Public/OnlineBeaconHost.h
--- a/Engine/Plugins/Online/OnlineSubsystemUtils/Source/OnlineSubsystemUtils/Public/OnlineBeaconHost.h
+++ b/Engine/Plugins/Online/OnlineSubsystemUtils/Source/OnlineSubsystemUtils/Public/OnlineBeaconHost.h
@@ -207,9 +207,6 @@ private:
 	void FinishHandshake(UNetConnection* Connection, FString BeaconType);
 	void SendFailurePacket(UNetConnection* Connection, FNetCloseResult&& CloseReason, const FText& ErrorText);
 
-	UE_DEPRECATED(5.1, "SendFailurePacket without CloseReason is deprecated. Use the version which takes CloseReason.")
-	void SendFailurePacket(UNetConnection* Connection, const FText& ErrorMsg);
-
 	void CloseHandshakeConnection(UNetConnection* Connection);
 
 	bool GetConnectionDataForUniqueNetId(const FUniqueNetId& UniqueNetId, UNetConnection*& OutConnection, FConnectionState*& OutConnectionState);
diff --git a/Engine/Plugins/Online/OnlineSubsystemUtils/Source/OnlineSubsystemUtils/Public/OnlineSubsystemUtils.h b/Engine/Plugins/Online/OnlineSubsystemUtils/Source/OnlineSubsystemUtils/Public/OnlineSubsystemUtils.h
--- a/Engine/Plugins/Online/OnlineSubsystemUtils/Source/OnlineSubsystemUtils/Public/OnlineSubsystemUtils.h
+++ b/Engine/Plugins/Online/OnlineSubsystemUtils/Source/OnlineSubsystemUtils/Public/OnlineSubsystemUtils.h
@@ -33,10 +33,6 @@ struct FVoiceSettings;
 /** @return an initialized audio component specifically for use with VoIP */
 ONLINESUBSYSTEMUTILS_API UAudioComponent* CreateVoiceAudioComponent(uint32 SampleRate, int32 NumChannels);
 
-/** @return an initialized Synth component specifically for use with VoIP */
-UE_DEPRECATED(4.25, "CreateVoiceSynthComponent(uint32 SampleRate) is deprecated, use CreateVoiceSynthComponent(UWorld* World, uint32 SampleRate).")
-ONLINESUBSYSTEMUTILS_API UVoipListenerSynthComponent* CreateVoiceSynthComponent(uint32 SampleRate);
-
 /** @return an initialized Synth component specifically for use with VoIP */
 ONLINESUBSYSTEMUTILS_API UVoipListenerSynthComponent* CreateVoiceSynthComponent(UWorld* World, uint32 SampleRate);
 
