Add read-write lock around FOnlineSubsystemModule::OnlineSubsystems
Fixes race conditions when multiple threads are accessing FOnlineSubsystemModule::GetOnlineSubsystem, observed to happen as part of multi-threaded FUniqueNetIdRepl deserialization

diff --git a/Engine/Plugins/Online/OnlineSubsystem/Source/Private/OnlineSubsystemModule.cpp b/Engine/Plugins/Online/OnlineSubsystem/Source/Private/OnlineSubsystemModule.cpp
--- a/Engine/Plugins/Online/OnlineSubsystem/Source/Private/OnlineSubsystemModule.cpp
+++ b/Engine/Plugins/Online/OnlineSubsystem/Source/Private/OnlineSubsystemModule.cpp
@@ -139,35 +139,35 @@ void FOnlineSubsystemModule::ProcessConfigDefinedModuleRedirects()
 	}
 }
 
-bool FOnlineSubsystemModule::TryLoadSubsystemAndSetDefault(FName ModuleName)
+bool FOnlineSubsystemModule::TryLoadSubsystemAndSetDefault(FName SubsystemName)
 {
 	// A module loaded with its factory method set for creation and a default instance of the online subsystem is required
 	bool bLoaded = false;
-	FString ModuleNameString = ModuleName.ToString();
-	if (LoadSubsystemModule(ModuleNameString, ModuleRedirects))
+	FString SubsystemNameString = SubsystemName.ToString();
+	if (LoadSubsystemModule(SubsystemNameString, ModuleRedirects))
 	{
-		if (OnlineFactories.Contains(ModuleName))
+		if (OnlineFactories.Contains(SubsystemName))
 		{
-			IOnlineSubsystem* OnlineSubsystem = GetOnlineSubsystem(ModuleName);
+			IOnlineSubsystem* OnlineSubsystem = GetOnlineSubsystem(SubsystemName);
 			if (OnlineSubsystem != nullptr)
 			{
-				UE_LOG_ONLINE(Log, TEXT("TryLoadSubsystemAndSetDefault: Loaded subsystem for module [%s]"), *ModuleNameString);
-				DefaultPlatformService = ModuleName;
+				UE_LOG_ONLINE(Log, TEXT("TryLoadSubsystemAndSetDefault: Loaded subsystem for type [%s]"), *SubsystemNameString);
+				DefaultPlatformService = SubsystemName;
 				bLoaded = true;
 			}
 			else
 			{
-				//UE_LOG_ONLINE(Warning, TEXT("TryLoadSubsystemAndSetDefault: GetOnlineSubsystem([%s]) failed"), *ModuleNameString);
+				//UE_LOG_ONLINE(Warning, TEXT("TryLoadSubsystemAndSetDefault: GetOnlineSubsystem([%s]) failed"), *SubsystemNameString);
 			}
 		}
 		else
 		{
-			UE_LOG_ONLINE(Warning, TEXT("TryLoadSubsystemAndSetDefault: OnlineFactories does not contain [%s]"), *ModuleNameString);
+			UE_LOG_ONLINE(Warning, TEXT("TryLoadSubsystemAndSetDefault: OnlineFactories does not contain [%s]"), *SubsystemNameString);
 		}
 	}
 	else
 	{
-		UE_LOG_ONLINE(Warning, TEXT("TryLoadSubsystemAndSetDefault: LoadSubsystemModule([%s]) failed"), *ModuleNameString);
+		UE_LOG_ONLINE(Warning, TEXT("TryLoadSubsystemAndSetDefault: LoadSubsystemModule([%s]) failed"), *SubsystemNameString);
 	}
 
 	return bLoaded;
@@ -208,6 +208,7 @@ void FOnlineSubsystemModule::ReloadDefaultSubsystem()
 void FOnlineSubsystemModule::PreUnloadOnlineSubsystem()
 {
 	// Shutdown all online subsystem instances
+	FReadScopeLock ReadLock(OnlineSubsystemsLock);
 	for (TMap<FName, IOnlineSubsystemPtr>::TIterator It(OnlineSubsystems); It; ++It)
 	{
 		It.Value()->PreUnload();
@@ -219,11 +220,14 @@ void FOnlineSubsystemModule::ShutdownOnlineSubsystem()
 	FModuleManager& ModuleManager = FModuleManager::Get();
 
 	// Shutdown all online subsystem instances
-	for (TMap<FName, IOnlineSubsystemPtr>::TIterator It(OnlineSubsystems); It; ++It)
 	{
-		It.Value()->Shutdown();
+		FWriteScopeLock WriteLock(OnlineSubsystemsLock);
+		for (TMap<FName, IOnlineSubsystemPtr>::TIterator It(OnlineSubsystems); It; ++It)
+		{
+			It.Value()->Shutdown();
+		}
+		OnlineSubsystems.Empty();
 	}
-	OnlineSubsystems.Empty();
 
 	// Unload all the supporting factories
 	for (TMap<FName, IOnlineFactory*>::TIterator It(OnlineFactories); It; ++It)
@@ -257,6 +261,7 @@ void FOnlineSubsystemModule::UnregisterPlatformService(const FName FactoryName)
 
 void FOnlineSubsystemModule::EnumerateOnlineSubsystems(FEnumerateOnlineSubsystemCb& EnumCb)
 {
+	FReadScopeLock ReadLock(OnlineSubsystemsLock);
 	for (TPair<FName, IOnlineSubsystemPtr>& OnlineSubsystem : OnlineSubsystems)
 	{
 		if (OnlineSubsystem.Value.IsValid())
@@ -311,6 +316,34 @@ FName FOnlineSubsystemModule::ParseOnlineSubsystemName(const FName& FullName, FN
 	return Entry->FullPath;
 }
 
+namespace
+{
+IOnlineSubsystemPtr FindExistingSubsystem(const TMap<FName, IOnlineSubsystemPtr>& OnlineSubsystems, const FName& KeyName)
+{
+	IOnlineSubsystemPtr Result;
+	if (const IOnlineSubsystemPtr* ExistingOnlineSubsystem = OnlineSubsystems.Find(KeyName))
+	{
+		Result = *ExistingOnlineSubsystem;
+	}
+	return Result;
+}
+
+IOnlineFactory* FindOrLoadOnlineFactory(const TMap<FName, IOnlineFactory*>& OnlineFactories, const TMap<FString, FName>& ModuleRedirects, const FName& SubsystemName)
+{
+	IOnlineFactory* const * ExistingFactory = OnlineFactories.Find(SubsystemName);
+	if (ExistingFactory == nullptr)
+	{
+		// Attempt to load the requested factory. This is expected to modify OnlineFactories.
+		if (LoadSubsystemModule(SubsystemName.ToString(), ModuleRedirects))
+		{
+			// If the module loaded successfully this should be non-NULL
+			ExistingFactory = OnlineFactories.Find(SubsystemName);
+		}
+	}
+	return ExistingFactory ? *ExistingFactory : nullptr;
+}
+} // anonymous
+
 IOnlineSubsystem* FOnlineSubsystemModule::GetOnlineSubsystem(const FName InSubsystemName)
 {
 	FName SubsystemName, InstanceName;
@@ -319,50 +352,48 @@ IOnlineSubsystem* FOnlineSubsystemModule::GetOnlineSubsystem(const FName InSubsy
 	IOnlineSubsystemPtr OnlineSubsystem;
 	if (!SubsystemName.IsNone())
 	{
-		IOnlineSubsystemPtr* ExistingOnlineSubsystem = OnlineSubsystems.Find(KeyName);
-		if (ExistingOnlineSubsystem != nullptr)
 		{
-			OnlineSubsystem = *ExistingOnlineSubsystem;
+			FReadScopeLock ReadLock(OnlineSubsystemsLock);
+			OnlineSubsystem = FindExistingSubsystem(OnlineSubsystems, KeyName);
 		}
-		else
+		bool bWasNewlyCreated = false;
+		if (!OnlineSubsystem)
 		{
-			if (IOnlineSubsystem::IsEnabled(SubsystemName))
+			FWriteScopeLock WriteLock(OnlineSubsystemsLock);
+			// Ensure that an OnlineSubsystem was not added in the narrow window between changing from a read lock to a write lock
+			OnlineSubsystem = FindExistingSubsystem(OnlineSubsystems, KeyName);
+			if (!OnlineSubsystem)
 			{
-				IOnlineFactory** OSSFactory = OnlineFactories.Find(SubsystemName);
-				if (OSSFactory == nullptr)
+				if (IOnlineSubsystem::IsEnabled(SubsystemName))
 				{
-					// Attempt to load the requested factory
-					IModuleInterface* NewModule = LoadSubsystemModule(SubsystemName.ToString(), ModuleRedirects);
-					if (NewModule)
+					IOnlineFactory* OSSFactory = FindOrLoadOnlineFactory(OnlineFactories, ModuleRedirects, SubsystemName);
+					if (OSSFactory != nullptr)
 					{
-						// If the module loaded successfully this should be non-NULL
-						OSSFactory = OnlineFactories.Find(SubsystemName);
-					}
-				}
-
-				if (OSSFactory != nullptr)
-				{
-					IOnlineSubsystemPtr NewSubsystemInstance = (*OSSFactory)->CreateSubsystem(InstanceName);
-					if (NewSubsystemInstance.IsValid())
-					{
-						UE_LOG_ONLINE(Log, TEXT("Created online subsystem instance for: %s"), *InSubsystemName.ToString());
+						OnlineSubsystem = OSSFactory->CreateSubsystem(InstanceName);
+						if (OnlineSubsystem.IsValid())
+						{
+							UE_LOG_ONLINE(Log, TEXT("Created online subsystem instance for: %s"), *InSubsystemName.ToString());
 
-						OnlineSubsystems.Add(KeyName, NewSubsystemInstance);
-						OnlineSubsystem = NewSubsystemInstance;
-						FOnlineSubsystemDelegates::OnOnlineSubsystemCreated.Broadcast(NewSubsystemInstance.Get());						
-					}
-					else
-					{
-						bool* bNotedPreviously = OnlineSubsystemFailureNotes.Find(KeyName);
-						if (!bNotedPreviously || !(*bNotedPreviously))
+							OnlineSubsystems.Add(KeyName, OnlineSubsystem);
+							bWasNewlyCreated = true;
+						}
+						else
 						{
-							UE_LOG_ONLINE(Log, TEXT("Unable to create OnlineSubsystem instance %s"), (InstanceName == FOnlineSubsystemImpl::DefaultInstanceName) ? *SubsystemName.ToString() : *KeyName.ToString());
-							OnlineSubsystemFailureNotes.Add(KeyName, true);
+							const bool bNotedPreviously = OnlineSubsystemFailureNotes.Contains(KeyName);
+							if (!bNotedPreviously)
+							{
+								UE_LOG_ONLINE(Log, TEXT("Unable to create OnlineSubsystem instance %s"), (InstanceName == FOnlineSubsystemImpl::DefaultInstanceName) ? *SubsystemName.ToString() : *KeyName.ToString());
+								OnlineSubsystemFailureNotes.Add(KeyName);
+							}
 						}
 					}
 				}
 			}
 		}
+		if (bWasNewlyCreated)
+		{
+			FOnlineSubsystemDelegates::OnOnlineSubsystemCreated.Broadcast(OnlineSubsystem.Get());
+		}
 	}
 
 	return OnlineSubsystem.Get();
@@ -403,7 +434,10 @@ void FOnlineSubsystemModule::DestroyOnlineSubsystem(const FName InSubsystemName)
 	if (!SubsystemName.IsNone())
 	{
 		IOnlineSubsystemPtr OnlineSubsystem;
-		OnlineSubsystems.RemoveAndCopyValue(KeyName, OnlineSubsystem);
+		{
+			FWriteScopeLock WriteLock(OnlineSubsystemsLock);
+			OnlineSubsystems.RemoveAndCopyValue(KeyName, OnlineSubsystem);
+		}
 		if (OnlineSubsystem.IsValid())
 		{
 			OnlineSubsystem->Shutdown();
@@ -424,8 +458,12 @@ bool FOnlineSubsystemModule::DoesInstanceExist(const FName InSubsystemName) cons
 	FName KeyName = ParseOnlineSubsystemName(InSubsystemName, SubsystemName, InstanceName);
 	if (!SubsystemName.IsNone())
 	{
-		const IOnlineSubsystemPtr* OnlineSubsystem = OnlineSubsystems.Find(KeyName);
-		return OnlineSubsystem && OnlineSubsystem->IsValid() ? true : false;
+		IOnlineSubsystemPtr OnlineSubsystem;
+		{
+			FReadScopeLock ReadLock(OnlineSubsystemsLock);
+			OnlineSubsystem = FindExistingSubsystem(OnlineSubsystems, KeyName);
+		}
+		return OnlineSubsystem.IsValid();
 	}
 
 	return false;
diff --git a/Engine/Plugins/Online/OnlineSubsystem/Source/Public/OnlineSubsystemModule.h b/Engine/Plugins/Online/OnlineSubsystem/Source/Public/OnlineSubsystemModule.h
--- a/Engine/Plugins/Online/OnlineSubsystem/Source/Public/OnlineSubsystemModule.h
+++ b/Engine/Plugins/Online/OnlineSubsystem/Source/Public/OnlineSubsystemModule.h
@@ -6,12 +6,9 @@
 #include "Templates/SharedPointer.h"
 
 class IOnlineFactory;
-
 class IOnlineSubsystem;
-
 typedef TSharedPtr<class IOnlineSubsystem, ESPMode::ThreadSafe> IOnlineSubsystemPtr;
 
-
 /**
  * Online subsystem module class
  * Wraps the loading of an online subsystem by name and allows new services to register themselves for use
@@ -44,13 +41,20 @@ private:
 	TMap<FString, FName> ConfigDefinedSubsystems;
 
 	/** Existing instances of any online subsystems created <PlatformName:InstanceName> */
-	TMap<FName, class IOnlineFactory*> OnlineFactories;
+	TMap<FName, IOnlineFactory*> OnlineFactories;
 
 	/** Mapping of all currently loaded platform service subsystems to their name */
 	TMap<FName, IOnlineSubsystemPtr> OnlineSubsystems;
 
+	/**
+	 * Read write lock for access to OnlineSubsystems.
+	 * The intent is OnlineSubsystems can have elements added on separate threads through normal gameplay, though this is expected to be very rare and only once per subsystem type.
+	 * OnlineSubsystems is not expected to remove elements except in very controlled scenarios.
+	 */
+	mutable FRWLock OnlineSubsystemsLock;
+
 	/** Have we warned already for a given online subsystem creation failure */
-	TMap<FName, bool> OnlineSubsystemFailureNotes;
+	TSet<FName> OnlineSubsystemFailureNotes;
 
 	/** Config driven override of module name for online subsystem */
 	TMap<FString, FName> ModuleRedirects;
