Change read-write lock to critical section.
The write lock obtained while loading a new OSS will block a new read lock on the same thread.  This has been observed when an OnlineSubsystem is loading another OnlineSubsystem.  Critical sections allow re-entry.

diff --git a/Engine/Plugins/Online/OnlineSubsystem/Source/Private/OnlineSubsystemModule.cpp b/Engine/Plugins/Online/OnlineSubsystem/Source/Private/OnlineSubsystemModule.cpp
--- a/Engine/Plugins/Online/OnlineSubsystem/Source/Private/OnlineSubsystemModule.cpp
+++ b/Engine/Plugins/Online/OnlineSubsystem/Source/Private/OnlineSubsystemModule.cpp
@@ -208,7 +208,7 @@ void FOnlineSubsystemModule::ReloadDefaultSubsystem()
 void FOnlineSubsystemModule::PreUnloadOnlineSubsystem()
 {
 	// Shutdown all online subsystem instances
-	FReadScopeLock ReadLock(OnlineSubsystemsLock);
+	FScopeLock Lock(&OnlineSubsystemsLock);
 	for (TMap<FName, IOnlineSubsystemPtr>::TIterator It(OnlineSubsystems); It; ++It)
 	{
 		It.Value()->PreUnload();
@@ -221,7 +221,7 @@ void FOnlineSubsystemModule::ShutdownOnlineSubsystem()
 
 	// Shutdown all online subsystem instances
 	{
-		FWriteScopeLock WriteLock(OnlineSubsystemsLock);
+		FScopeLock Lock(&OnlineSubsystemsLock);
 		for (TMap<FName, IOnlineSubsystemPtr>::TIterator It(OnlineSubsystems); It; ++It)
 		{
 			It.Value()->Shutdown();
@@ -261,7 +261,7 @@ void FOnlineSubsystemModule::UnregisterPlatformService(const FName FactoryName)
 
 void FOnlineSubsystemModule::EnumerateOnlineSubsystems(FEnumerateOnlineSubsystemCb& EnumCb)
 {
-	FReadScopeLock ReadLock(OnlineSubsystemsLock);
+	FScopeLock Lock(&OnlineSubsystemsLock);
 	for (TPair<FName, IOnlineSubsystemPtr>& OnlineSubsystem : OnlineSubsystems)
 	{
 		if (OnlineSubsystem.Value.IsValid())
@@ -352,15 +352,9 @@ IOnlineSubsystem* FOnlineSubsystemModule::GetOnlineSubsystem(const FName InSubsy
 	IOnlineSubsystemPtr OnlineSubsystem;
 	if (!SubsystemName.IsNone())
 	{
-		{
-			FReadScopeLock ReadLock(OnlineSubsystemsLock);
-			OnlineSubsystem = FindExistingSubsystem(OnlineSubsystems, KeyName);
-		}
 		bool bWasNewlyCreated = false;
-		if (!OnlineSubsystem)
 		{
-			FWriteScopeLock WriteLock(OnlineSubsystemsLock);
-			// Ensure that an OnlineSubsystem was not added in the narrow window between changing from a read lock to a write lock
+			FScopeLock Lock(&OnlineSubsystemsLock);
 			OnlineSubsystem = FindExistingSubsystem(OnlineSubsystems, KeyName);
 			if (!OnlineSubsystem)
 			{
@@ -435,7 +429,7 @@ void FOnlineSubsystemModule::DestroyOnlineSubsystem(const FName InSubsystemName)
 	{
 		IOnlineSubsystemPtr OnlineSubsystem;
 		{
-			FWriteScopeLock WriteLock(OnlineSubsystemsLock);
+			FScopeLock Lock(&OnlineSubsystemsLock);
 			OnlineSubsystems.RemoveAndCopyValue(KeyName, OnlineSubsystem);
 		}
 		if (OnlineSubsystem.IsValid())
@@ -460,7 +454,7 @@ bool FOnlineSubsystemModule::DoesInstanceExist(const FName InSubsystemName) cons
 	{
 		IOnlineSubsystemPtr OnlineSubsystem;
 		{
-			FReadScopeLock ReadLock(OnlineSubsystemsLock);
+			FScopeLock Lock(&OnlineSubsystemsLock);
 			OnlineSubsystem = FindExistingSubsystem(OnlineSubsystems, KeyName);
 		}
 		return OnlineSubsystem.IsValid();
diff --git a/Engine/Plugins/Online/OnlineSubsystem/Source/Public/OnlineSubsystemModule.h b/Engine/Plugins/Online/OnlineSubsystem/Source/Public/OnlineSubsystemModule.h
--- a/Engine/Plugins/Online/OnlineSubsystem/Source/Public/OnlineSubsystemModule.h
+++ b/Engine/Plugins/Online/OnlineSubsystem/Source/Public/OnlineSubsystemModule.h
@@ -47,11 +47,11 @@ private:
 	TMap<FName, IOnlineSubsystemPtr> OnlineSubsystems;
 
 	/**
-	 * Read write lock for access to OnlineSubsystems.
+	 * Lock for access to OnlineSubsystems array to allow for multiple threads to be able to create/get online subsystems simultaneously.
 	 * The intent is OnlineSubsystems can have elements added on separate threads through normal gameplay, though this is expected to be very rare and only once per subsystem type.
-	 * OnlineSubsystems is not expected to remove elements except in very controlled scenarios.
+	 * OnlineSubsystems is not expected to remove elements except in very controlled scenarios.  In that case, the lock is only limiting access to the OnlineSubsystems array, but the result may be invalidated.
 	 */
-	mutable FRWLock OnlineSubsystemsLock;
+	mutable FCriticalSection OnlineSubsystemsLock;
 
 	/** Have we warned already for a given online subsystem creation failure */
 	TSet<FName> OnlineSubsystemFailureNotes;
