Fixed mismatched copy constructor/assignment operator in FOnlinePartyTypeId.

#okforgtihub public

diff --git a/Engine/Plugins/Online/OnlineSubsystem/Source/Public/OnlineSubsystemTypes.h b/Engine/Plugins/Online/OnlineSubsystem/Source/Public/OnlineSubsystemTypes.h
--- a/Engine/Plugins/Online/OnlineSubsystem/Source/Public/OnlineSubsystemTypes.h
+++ b/Engine/Plugins/Online/OnlineSubsystem/Source/Public/OnlineSubsystemTypes.h
@@ -1536,7 +1536,6 @@ class FOnlinePartyTypeId
 public:
 	typedef uint32 TInternalType;
 	explicit FOnlinePartyTypeId(const TInternalType InValue = 0) : Value(InValue) {}
-	FOnlinePartyTypeId(const FOnlinePartyTypeId& Other) : Value(Other.Value) {}
 
 	bool operator==(const FOnlinePartyTypeId Rhs) const { return Value == Rhs.Value; }
 	bool operator!=(const FOnlinePartyTypeId Rhs) const { return Value != Rhs.Value; }
