OSS Null now properly assigns and checks the BuildUniqueId field in sessions.

diff --git a/Engine/Plugins/Online/OnlineSubsystem/Source/Public/OnlineSessionSettings.h b/Engine/Plugins/Online/OnlineSubsystem/Source/Public/OnlineSessionSettings.h
--- a/Engine/Plugins/Online/OnlineSubsystem/Source/Public/OnlineSessionSettings.h
+++ b/Engine/Plugins/Online/OnlineSubsystem/Source/Public/OnlineSessionSettings.h
@@ -2,8 +2,8 @@
 
 #pragma once
 
-#include "OnlineSubsystemTypes.h"
 #include "OnlineKeyValuePair.h"
+#include "OnlineSubsystem.h"
 
 // NOTE:  Session setting names have been moved to OnlineBase in OnlineSessionNames.h
 
@@ -280,7 +280,7 @@ public:
 		, bAntiCheatProtected(false)
 		, bUseLobbiesIfAvailable(false)
 		, bUseLobbiesVoiceChatIfAvailable(false)
-		, BuildUniqueId(0)
+		, BuildUniqueId(GetBuildUniqueId())
 	{
 		// Example usage of settings
 //		Set(SETTING_MAPNAME, FString(TEXT("")), EOnlineDataAdvertisementType::ViaOnlineService);
diff --git a/Engine/Plugins/Online/OnlineSubsystemNull/Source/Private/OnlineSessionInterfaceNull.cpp b/Engine/Plugins/Online/OnlineSubsystemNull/Source/Private/OnlineSessionInterfaceNull.cpp
--- a/Engine/Plugins/Online/OnlineSubsystemNull/Source/Private/OnlineSessionInterfaceNull.cpp
+++ b/Engine/Plugins/Online/OnlineSubsystemNull/Source/Private/OnlineSessionInterfaceNull.cpp
@@ -216,9 +216,6 @@ bool FOnlineSessionNull::CreateSession(int32 HostingPlayerNum, FName SessionName
 			Session->OwningUserId = FUniqueNetIdNull::Create(FString::Printf(TEXT("%d"), HostingPlayerNum));
 			Session->OwningUserName = FString(TEXT("NullUser"));
 		}
-		
-		// Unique identifier of this build for compatibility
-		Session->SessionSettings.BuildUniqueId = GetBuildUniqueId();
 
 		// Setup the host session info
 		FOnlineSessionInfoNull* NewSessionInfo = new FOnlineSessionInfoNull();
