Add handling for canceling account login.

diff --git a/Engine/Plugins/Online/OnlineSubsystem/Source/Public/OnlineIdentityErrors.h b/Engine/Plugins/Online/OnlineSubsystem/Source/Public/OnlineIdentityErrors.h
--- a/Engine/Plugins/Online/OnlineSubsystem/Source/Public/OnlineIdentityErrors.h
+++ b/Engine/Plugins/Online/OnlineSubsystem/Source/Public/OnlineIdentityErrors.h
@@ -25,6 +25,8 @@ namespace OnlineIdentity
 		inline FOnlineError PinGrantFailure() { return ONLINE_ERROR(EOnlineErrorResult::FailExtended, TEXT("pin_grant_failure")); }
 		inline FOnlineError PinGrantTimeout() { return ONLINE_ERROR(EOnlineErrorResult::FailExtended, TEXT("pin_grant_timeout")); }
 
+		inline FOnlineError Canceled() { return ONLINE_ERROR(EOnlineErrorResult::Canceled); }
+
 		// Params
 		UE_DEPRECATED(5.2, "AuthLoginParam is deprecated, please use UE_ONLINE_ERROR_PARAM_IDENTITY_AUTH_LOGIN instead.")
 		extern ONLINESUBSYSTEM_API const FString AuthLoginParam;
