Convert online subsystem defines to be global FNames to avoid a lookup every time we use one.

diff --git a/Engine/Plugins/Online/OnlineSubsystem/Source/Private/OnlineSubsystem.cpp b/Engine/Plugins/Online/OnlineSubsystem/Source/Private/OnlineSubsystem.cpp
--- a/Engine/Plugins/Online/OnlineSubsystem/Source/Private/OnlineSubsystem.cpp
+++ b/Engine/Plugins/Online/OnlineSubsystem/Source/Private/OnlineSubsystem.cpp
@@ -59,10 +59,86 @@ ONLINESUBSYSTEM_API DEFINE_STAT(STAT_Session_Interface);
 ONLINESUBSYSTEM_API DEFINE_STAT(STAT_Voice_Interface);
 #endif
 
+#ifndef NULL_SUBSYSTEM
+const FName NULL_SUBSYSTEM(TEXT("NULL"));
+#endif
+
+#ifndef GOOGLEPLAY_SUBSYSTEM
+const FName GOOGLEPLAY_SUBSYSTEM(TEXT("GOOGLEPLAY"));
+#endif
+
+#ifndef IOS_SUBSYSTEM
+const FName IOS_SUBSYSTEM(TEXT("IOS"));
+#endif
+
+#ifndef APPLE_SUBSYSTEM
+const FName APPLE_SUBSYSTEM(TEXT("APPLE"));
+#endif
+
+#ifndef AMAZON_SUBSYSTEM
+const FName AMAZON_SUBSYSTEM(TEXT("AMAZON"));
+#endif
+
+#ifndef GOOGLE_SUBSYSTEM
+const FName GOOGLE_SUBSYSTEM(TEXT("GOOGLE"));
+#endif
+
+#ifndef FACEBOOK_SUBSYSTEM
+const FName FACEBOOK_SUBSYSTEM(TEXT("FACEBOOK"));
+#endif
+
+#ifndef GAMECIRCLE_SUBSYSTEM
+const FName GAMECIRCLE_SUBSYSTEM(TEXT("GAMECIRCLE"));
+#endif
+
+#ifndef STEAM_SUBSYSTEM
+const FName STEAM_SUBSYSTEM(TEXT("STEAM"));
+#endif
+
+#ifndef PS4_SUBSYSTEM
+const FName PS4_SUBSYSTEM(TEXT("PS4"));
+#endif
+
+#ifndef PS4SERVER_SUBSYSTEM
+const FName PS4SERVER_SUBSYSTEM(TEXT("PS4SERVER"));
+#endif
+
+#ifndef THUNDERHEAD_SUBSYSTEM
+const FName THUNDERHEAD_SUBSYSTEM(TEXT("THUNDERHEAD"));
+#endif
+
+#ifndef MCP_SUBSYSTEM
+const FName MCP_SUBSYSTEM(TEXT("MCP"));
+#endif
+
+#ifndef MCP_SUBSYSTEM_EMBEDDED
+const FName MCP_SUBSYSTEM_EMBEDDED(TEXT("MCP:EMBEDDED"));
+#endif
+
+#ifndef TENCENT_SUBSYSTEM
+const FName TENCENT_SUBSYSTEM(TEXT("TENCENT"));
+#endif
+
+#ifndef SWITCH_SUBSYSTEM
+const FName SWITCH_SUBSYSTEM(TEXT("SWITCH"));
+#endif
+
 PRAGMA_DISABLE_DEPRECATION_WARNINGS
-FName OCULUS_SUBSYSTEM(TEXT("OCULUS"));
+const FName OCULUS_SUBSYSTEM(TEXT("OCULUS"));
 PRAGMA_ENABLE_DEPRECATION_WARNINGS
 
+#ifndef SAMSUNG_SUBSYSTEM
+const FName SAMSUNG_SUBSYSTEM(TEXT("SAMSUNG"));
+#endif
+
+#ifndef QUAIL_SUBSYSTEM
+const FName QUAIL_SUBSYSTEM(TEXT("Quail"));
+#endif
+
+#ifndef EOS_SUBSYSTEM
+const FName EOS_SUBSYSTEM(TEXT("EOS"));
+#endif
+
 /** The default key that will update presence text in the platform's UI */
 const FString DefaultPresenceKey = TEXT("RichPresence");
 
diff --git a/Engine/Plugins/Online/OnlineSubsystem/Source/Public/OnlineSubsystemNames.h b/Engine/Plugins/Online/OnlineSubsystem/Source/Public/OnlineSubsystemNames.h
--- a/Engine/Plugins/Online/OnlineSubsystem/Source/Public/OnlineSubsystemNames.h
+++ b/Engine/Plugins/Online/OnlineSubsystem/Source/Public/OnlineSubsystemNames.h
@@ -18,81 +18,81 @@
 
 
 #ifndef NULL_SUBSYSTEM
-#define NULL_SUBSYSTEM FName(TEXT("NULL"))
+extern ONLINESUBSYSTEM_API const FName NULL_SUBSYSTEM;
 #endif
 
 #ifndef GOOGLEPLAY_SUBSYSTEM
-#define GOOGLEPLAY_SUBSYSTEM FName(TEXT("GOOGLEPLAY"))
+extern ONLINESUBSYSTEM_API const FName GOOGLEPLAY_SUBSYSTEM;
 #endif
 
 #ifndef IOS_SUBSYSTEM
-#define IOS_SUBSYSTEM FName(TEXT("IOS"))
+extern ONLINESUBSYSTEM_API const FName IOS_SUBSYSTEM;
 #endif
 
 #ifndef APPLE_SUBSYSTEM
-#define APPLE_SUBSYSTEM FName(TEXT("APPLE"))
+extern ONLINESUBSYSTEM_API const FName APPLE_SUBSYSTEM;
 #endif
 
 #ifndef AMAZON_SUBSYSTEM
-#define AMAZON_SUBSYSTEM FName(TEXT("AMAZON"))
+extern ONLINESUBSYSTEM_API const FName AMAZON_SUBSYSTEM;
 #endif
 
 #ifndef GOOGLE_SUBSYSTEM
-#define GOOGLE_SUBSYSTEM FName(TEXT("GOOGLE"))
+extern ONLINESUBSYSTEM_API const FName GOOGLE_SUBSYSTEM;
 #endif
 
 #ifndef FACEBOOK_SUBSYSTEM
-#define FACEBOOK_SUBSYSTEM FName(TEXT("FACEBOOK"))
+extern ONLINESUBSYSTEM_API const FName FACEBOOK_SUBSYSTEM;
 #endif
 
 #ifndef GAMECIRCLE_SUBSYSTEM
-#define GAMECIRCLE_SUBSYSTEM FName(TEXT("GAMECIRCLE"))
+extern ONLINESUBSYSTEM_API const FName GAMECIRCLE_SUBSYSTEM;
 #endif
 
 #ifndef STEAM_SUBSYSTEM
-#define STEAM_SUBSYSTEM FName(TEXT("STEAM"))
+extern ONLINESUBSYSTEM_API const FName STEAM_SUBSYSTEM;
 #endif
 
 #ifndef PS4_SUBSYSTEM
-#define PS4_SUBSYSTEM FName(TEXT("PS4"))
+extern ONLINESUBSYSTEM_API const FName PS4_SUBSYSTEM;
 #endif
 
 #ifndef PS4SERVER_SUBSYSTEM
-#define PS4SERVER_SUBSYSTEM FName(TEXT("PS4SERVER"))
+extern ONLINESUBSYSTEM_API const FName PS4SERVER_SUBSYSTEM;
 #endif
 
 #ifndef THUNDERHEAD_SUBSYSTEM
-#define THUNDERHEAD_SUBSYSTEM FName(TEXT("THUNDERHEAD"))
+extern ONLINESUBSYSTEM_API const FName THUNDERHEAD_SUBSYSTEM;
 #endif
 
 #ifndef MCP_SUBSYSTEM
-#define MCP_SUBSYSTEM FName(TEXT("MCP"))
+extern ONLINESUBSYSTEM_API const FName MCP_SUBSYSTEM;
 #endif
 
 #ifndef MCP_SUBSYSTEM_EMBEDDED
-#define MCP_SUBSYSTEM_EMBEDDED FName(TEXT("MCP:EMBEDDED"))
+extern ONLINESUBSYSTEM_API const FName MCP_SUBSYSTEM_EMBEDDED;
 #endif
 
 #ifndef TENCENT_SUBSYSTEM
-#define TENCENT_SUBSYSTEM FName(TEXT("TENCENT"))
+extern ONLINESUBSYSTEM_API const FName TENCENT_SUBSYSTEM;
 #endif
 
 #ifndef SWITCH_SUBSYSTEM
-#define SWITCH_SUBSYSTEM FName(TEXT("SWITCH"))
+extern ONLINESUBSYSTEM_API const FName SWITCH_SUBSYSTEM;
 #endif
 
-UE_DEPRECATED(5.2, "OnlineSubsystemOculus has been deprecated and will be removed") ONLINESUBSYSTEM_API extern FName OCULUS_SUBSYSTEM;
+UE_DEPRECATED(5.2, "OnlineSubsystemOculus has been deprecated and will be removed") ONLINESUBSYSTEM_API extern const FName OCULUS_SUBSYSTEM;
 
 #ifndef SAMSUNG_SUBSYSTEM
-#define SAMSUNG_SUBSYSTEM FName(TEXT("SAMSUNG"))
+extern ONLINESUBSYSTEM_API const FName SAMSUNG_SUBSYSTEM;
 #endif
 
 #ifndef QUAIL_SUBSYSTEM
-#define QUAIL_SUBSYSTEM FName(TEXT("Quail"))
+extern ONLINESUBSYSTEM_API const FName QUAIL_SUBSYSTEM;
 #endif
 
 #ifndef EOS_SUBSYSTEM
-#define EOS_SUBSYSTEM FName(TEXT("EOS"))
+extern ONLINESUBSYSTEM_API const FName EOS_SUBSYSTEM;
 #endif
 
 #if UE_ENABLE_INCLUDE_ORDER_DEPRECATED_IN_5_2
