Un-inline functions in OnlineSubsystemTypes.h

diff --git a/Engine/Plugins/Online/OnlineSubsystem/Source/Private/OnlineSubsystemTypes.cpp b/Engine/Plugins/Online/OnlineSubsystem/Source/Private/OnlineSubsystemTypes.cpp
new file mode 100644
--- /dev/null
+++ b/Engine/Plugins/Online/OnlineSubsystem/Source/Private/OnlineSubsystemTypes.cpp
@@ -0,0 +1,558 @@
+// Copyright Epic Games, Inc. All Rights Reserved.
+
+#include "OnlineSubsystemTypes.h"
+
+#include "Online/CoreOnline.h"
+
+namespace EOnlineEnvironment
+{
+	const TCHAR* ToString(EOnlineEnvironment::Type EnvironmentType)
+	{
+		switch (EnvironmentType)
+		{
+		case Development: return TEXT("Development");
+		case Certification: return TEXT("Certification");
+		case Production: return TEXT("Production");
+		case Unknown: default: return TEXT("Unknown");
+		};
+	}
+} // namespace EOnlineEnvironment
+
+namespace ELoginStatus
+{
+	const TCHAR* ToString(ELoginStatus::Type EnumVal)
+	{
+		switch (EnumVal)
+		{
+		case NotLoggedIn:
+		{
+			return TEXT("NotLoggedIn");
+		}
+		case UsingLocalProfile:
+		{
+			return TEXT("UsingLocalProfile");
+		}
+		case LoggedIn:
+		{
+			return TEXT("LoggedIn");
+		}
+		}
+		return TEXT("");
+	}
+} // namespace ELoginStatus
+
+namespace EOnlineServerConnectionStatus
+{
+	const TCHAR* ToString(EOnlineServerConnectionStatus::Type EnumVal)
+	{
+		switch (EnumVal)
+		{
+		case Normal:
+		{
+			return TEXT("Normal");
+		}
+		case NotConnected:
+		{
+			return TEXT("NotConnected");
+		}
+		case Connected:
+		{
+			return TEXT("Connected");
+		}
+		case ConnectionDropped:
+		{
+			return TEXT("ConnectionDropped");
+		}
+		case NoNetworkConnection:
+		{
+			return TEXT("NoNetworkConnection");
+		}
+		case ServiceUnavailable:
+		{
+			return TEXT("ServiceUnavailable");
+		}
+		case UpdateRequired:
+		{
+			return TEXT("UpdateRequired");
+		}
+		case ServersTooBusy:
+		{
+			return TEXT("ServersTooBusy");
+		}
+		case DuplicateLoginDetected:
+		{
+			return TEXT("DuplicateLoginDetected");
+		}
+		case InvalidUser:
+		{
+			return TEXT("InvalidUser");
+		}
+		case NotAuthorized:
+		{
+			return TEXT("NotAuthorized");
+		}
+		case InvalidSession:
+		{
+			return TEXT("InvalidSession");
+		}
+		default:
+		{
+			return TEXT("Unknown");
+		}
+		}
+		return TEXT("");
+	}
+} // namespace EOnlineServerConnectionStatus
+
+namespace EFeaturePrivilegeLevel
+{
+	const TCHAR* ToString(EFeaturePrivilegeLevel::Type EnumVal)
+	{
+		switch (EnumVal)
+		{
+		case Undefined:
+		{
+			return TEXT("Undefined");
+		}
+		case Disabled:
+		{
+			return TEXT("Disabled");
+		}
+		case EnabledFriendsOnly:
+		{
+			return TEXT("EnabledFriendsOnly");
+		}
+		case Enabled:
+		{
+			return TEXT("Enabled");
+		}
+		}
+		return TEXT("");
+	}
+} // namespace EFeaturePrivilegeLevel
+
+namespace EOnlineAsyncTaskState
+{
+	const TCHAR* ToString(EOnlineAsyncTaskState::Type EnumVal)
+	{
+		switch (EnumVal)
+		{
+		case NotStarted:
+		{
+			return TEXT("NotStarted");
+		}
+		case InProgress:
+		{
+			return TEXT("InProgress");
+		}
+		case Done:
+		{
+			return TEXT("Done");
+		}
+		case Failed:
+		{
+			return TEXT("Failed");
+		}
+		}
+		return TEXT("");
+	}
+} // namespace EOnlineAsyncTaskState
+
+namespace EOnlineFriendState
+{
+	const TCHAR* ToString(EOnlineFriendState::Type EnumVal)
+	{
+		switch (EnumVal)
+		{
+		case Offline:
+		{
+			return TEXT("Offline");
+		}
+		case Online:
+		{
+			return TEXT("Online");
+		}
+		case Away:
+		{
+			return TEXT("Away");
+		}
+		case Busy:
+		{
+			return TEXT("Busy");
+		}
+		}
+		return TEXT("");
+	}
+} // namespace EOnlineFriendState
+
+namespace ELeaderboardSort
+{
+	const TCHAR* ToString(ELeaderboardSort::Type EnumVal)
+	{
+		switch (EnumVal)
+		{
+		case None:
+		{
+			return TEXT("None");
+		}
+		case Ascending:
+		{
+			return TEXT("Ascending");
+		}
+		case Descending:
+		{
+			return TEXT("Descending");
+		}
+		}
+		return TEXT("");
+	}
+} // namespace ELeaderboardSort
+
+namespace ELeaderboardFormat
+{
+	const TCHAR* ToString(ELeaderboardFormat::Type EnumVal)
+	{
+		switch (EnumVal)
+		{
+		case Number:
+		{
+			return TEXT("Number");
+		}
+		case Seconds:
+		{
+			return TEXT("Seconds");
+		}
+		case Milliseconds:
+		{
+			return TEXT("Milliseconds");
+		}
+		}
+		return TEXT("");
+	}
+} // namespace ELeaderboardFormat
+
+namespace ELeaderboardUpdateMethod
+{
+	const TCHAR* ToString(ELeaderboardUpdateMethod::Type EnumVal)
+	{
+		switch (EnumVal)
+		{
+		case KeepBest:
+		{
+			return TEXT("KeepBest");
+		}
+		case Force:
+		{
+			return TEXT("Force");
+		}
+		}
+		return TEXT("");
+	}
+} // namespace ELeaderboardUpdateMethod
+
+namespace EOnlineSessionState
+{
+	const TCHAR* ToString(EOnlineSessionState::Type EnumVal)
+	{
+		switch (EnumVal)
+		{
+
+		case NoSession:
+		{
+			return TEXT("NoSession");
+		}
+		case Creating:
+		{
+			return TEXT("Creating");
+		}
+		case Pending:
+		{
+			return TEXT("Pending");
+		}
+		case Starting:
+		{
+			return TEXT("Starting");
+		}
+		case InProgress:
+		{
+			return TEXT("InProgress");
+		}
+		case Ending:
+		{
+			return TEXT("Ending");
+		}
+		case Ended:
+		{
+			return TEXT("Ended");
+		}
+		case Destroying:
+		{
+			return TEXT("Destroying");
+		}
+		}
+		return TEXT("");
+	}
+} // namespace EOnlineSessionState
+
+namespace EOnlineDataAdvertisementType
+{
+	const TCHAR* ToString(EOnlineDataAdvertisementType::Type EnumVal)
+	{
+		switch (EnumVal)
+		{
+		case DontAdvertise:
+		{
+			return TEXT("DontAdvertise");
+		}
+		case ViaPingOnly:
+		{
+			return TEXT("ViaPingOnly");
+		}
+		case ViaOnlineService:
+		{
+			return TEXT("OnlineService");
+		}
+		case ViaOnlineServiceAndPing:
+		{
+			return TEXT("OnlineServiceAndPing");
+		}
+		}
+		return TEXT("");
+	}
+} // namespace EOnlineDataAdvertisementType
+
+namespace EOnlineComparisonOp
+{
+	const TCHAR* ToString(EOnlineComparisonOp::Type EnumVal)
+	{
+		switch (EnumVal)
+		{
+		case Equals:
+		{
+			return TEXT("Equals");
+		}
+		case NotEquals:
+		{
+			return TEXT("NotEquals");
+		}
+		case GreaterThan:
+		{
+			return TEXT("GreaterThan");
+		}
+		case GreaterThanEquals:
+		{
+			return TEXT("GreaterThanEquals");
+		}
+		case LessThan:
+		{
+			return TEXT("LessThan");
+		}
+		case LessThanEquals:
+		{
+			return TEXT("LessThanEquals");
+		}
+		case Near:
+		{
+			return TEXT("Near");
+		}
+		case In:
+		{
+			return TEXT("In");
+		}
+		case NotIn:
+		{
+			return TEXT("NotIn");
+		}
+		}
+		return TEXT("");
+	}
+} // namespace EOnlineComparisonOp
+
+namespace EOnlineCachedResult
+{
+	const TCHAR* ToString(EOnlineCachedResult::Type EnumVal)
+	{
+		switch (EnumVal)
+		{
+		case Success:
+		{
+			return TEXT("Success");
+		}
+		case NotFound:
+		{
+			return TEXT("NotFound");
+		}
+		}
+		return TEXT("");
+	}
+} // namespace EOnlineCachedResult
+
+const TCHAR* LexToString(EFriendInvitePolicy EnumVal)
+{
+	switch (EnumVal)
+	{
+	case EFriendInvitePolicy::Public: return TEXT("PUBLIC");
+	case EFriendInvitePolicy::Friends_of_Friends: return TEXT("FRIENDS_OF_FRIENDS");
+	case EFriendInvitePolicy::Private: return TEXT("PRIVATE");
+	default: return TEXT("Invalid");
+	}
+}
+
+void LexFromString(EFriendInvitePolicy& Value, const TCHAR* String)
+{
+	if (FCString::Stricmp(String, TEXT("PUBLIC")) == 0)
+	{
+		Value = EFriendInvitePolicy::Public;
+	}
+	else if (FCString::Stricmp(String, TEXT("FRIENDS_OF_FRIENDS")) == 0)
+	{
+		Value = EFriendInvitePolicy::Friends_of_Friends;
+	}
+	else if (FCString::Stricmp(String, TEXT("PRIVATE")) == 0)
+	{
+		Value = EFriendInvitePolicy::Private;
+	}
+	else
+	{
+		Value = EFriendInvitePolicy::InvalidOrMax;
+	}
+}
+
+bool IOnlinePlatformData::Compare(const IOnlinePlatformData& Other) const
+{
+	return (GetSize() == Other.GetSize()) &&
+		(FMemory::Memcmp(GetBytes(), Other.GetBytes(), GetSize()) == 0);
+}
+
+namespace EInviteStatus
+{
+	const TCHAR* ToString(EInviteStatus::Type EnumVal)
+	{
+		switch (EnumVal)
+		{
+		case Unknown:
+		{
+			return TEXT("Unknown");
+		}
+		case Accepted:
+		{
+			return TEXT("Accepted");
+		}
+		case PendingInbound:
+		{
+			return TEXT("PendingInbound");
+		}
+		case PendingOutbound:
+		{
+			return TEXT("PendingOutbound");
+		}
+		case Blocked:
+		{
+			return TEXT("Blocked");
+		}
+		case Suggested:
+		{
+			return TEXT("Suggested");
+		}
+		}
+		return TEXT("");
+	}
+} // namespace EInviteStatus
+
+const TCHAR* ToString(EOnlineSharingCategory CategoryType)
+{
+	switch (CategoryType)
+	{
+	case EOnlineSharingCategory::None:
+	{
+		return TEXT("Category undefined");
+	}
+	case EOnlineSharingCategory::ReadPosts:
+	{
+		return TEXT("ReadPosts");
+	}
+	case EOnlineSharingCategory::Friends:
+	{
+		return TEXT("Friends");
+	}
+	case EOnlineSharingCategory::Mailbox:
+	{
+		return TEXT("Mailbox");
+	}
+	case EOnlineSharingCategory::OnlineStatus:
+	{
+		return TEXT("Online Status");
+	}
+	case EOnlineSharingCategory::ProfileInfo:
+	{
+		return TEXT("Profile Information");
+	}
+	case EOnlineSharingCategory::LocationInfo:
+	{
+		return TEXT("Location Information");
+	}
+	case EOnlineSharingCategory::SubmitPosts:
+	{
+		return TEXT("SubmitPosts");
+	}
+	case EOnlineSharingCategory::ManageFriends:
+	{
+		return TEXT("ManageFriends");
+	}
+	case EOnlineSharingCategory::AccountAdmin:
+	{
+		return TEXT("Account Admin");
+	}
+	case EOnlineSharingCategory::Events:
+	{
+		return TEXT("Events");
+	}
+	}
+	return TEXT("");
+}
+
+const TCHAR* ToString(EOnlineStatusUpdatePrivacy PrivacyType)
+{
+	switch (PrivacyType)
+	{
+	case EOnlineStatusUpdatePrivacy::OnlyMe:
+		return TEXT("Only Me");
+	case EOnlineStatusUpdatePrivacy::OnlyFriends:
+		return TEXT("Only Friends");
+	case EOnlineStatusUpdatePrivacy::Everyone:
+		return TEXT("Everyone");
+	}
+	return TEXT("");
+}
+
+uint32 GetTypeHash(const FOnlinePartyId& Value)
+{
+	return CityHash32(reinterpret_cast<const char*>(Value.GetBytes()), Value.GetSize());
+}
+
+void ParseOnlineSubsystemConfigPairs(TArrayView<const FString> InEntries, TArray<TPair<FString, FString>>& OutPairs)
+{
+	OutPairs.Reserve(InEntries.Num());
+	// Takes on the pattern "(Key=Value)"
+	for (const FString& Entry : InEntries)
+	{
+		FString TrimmedConfigEntry = Entry.TrimStartAndEnd();
+		FString KeyString;
+		FString ValueString;
+
+		if (TrimmedConfigEntry.Left(1) == TEXT("("))
+		{
+			TrimmedConfigEntry.RightChopInline(1, false);
+		}
+		if (TrimmedConfigEntry.Right(1) == TEXT(")"))
+		{
+			TrimmedConfigEntry.LeftChopInline(1, false);
+		}
+		if (TrimmedConfigEntry.Split(TEXT("="), &KeyString, &ValueString, ESearchCase::CaseSensitive))
+		{
+			KeyString.TrimStartAndEndInline();
+			ValueString.TrimStartAndEndInline();
+		}
+		OutPairs.Emplace(MoveTemp(KeyString), MoveTemp(ValueString));
+	}
+}
diff --git a/Engine/Plugins/Online/OnlineSubsystem/Source/Public/OnlineSubsystemTypes.h b/Engine/Plugins/Online/OnlineSubsystem/Source/Public/OnlineSubsystemTypes.h
--- a/Engine/Plugins/Online/OnlineSubsystem/Source/Public/OnlineSubsystemTypes.h
+++ b/Engine/Plugins/Online/OnlineSubsystem/Source/Public/OnlineSubsystemTypes.h
@@ -41,16 +41,7 @@ namespace EOnlineEnvironment
 	};
 
 	/** @return the stringified version of the enum passed in */
-	inline const TCHAR* ToString(EOnlineEnvironment::Type EnvironmentType)
-	{
-		switch (EnvironmentType)
-		{
-			case Development: return TEXT("Development");
-			case Certification: return TEXT("Certification");
-			case Production: return TEXT("Production");
-			case Unknown: default: return TEXT("Unknown");
-		};
-	}
+	ONLINESUBSYSTEM_API const TCHAR* ToString(EOnlineEnvironment::Type EnvironmentType);
 }
 
 /** Possible login states */
@@ -67,25 +58,7 @@ namespace ELoginStatus
 	};
 
 	/** @return the stringified version of the enum passed in */
-	inline const TCHAR* ToString(ELoginStatus::Type EnumVal)
-	{
-		switch (EnumVal)
-		{
-			case NotLoggedIn:
-			{
-				return TEXT("NotLoggedIn");
-			}
-			case UsingLocalProfile:
-			{
-				return TEXT("UsingLocalProfile");
-			}
-			case LoggedIn:
-			{
-				return TEXT("LoggedIn");
-			}
-		}
-		return TEXT("");
-	}
+	ONLINESUBSYSTEM_API const TCHAR* ToString(ELoginStatus::Type EnumVal);
 };
 
 /** Possible connection states */
@@ -120,65 +93,7 @@ namespace EOnlineServerConnectionStatus
 	};
 
 	/** @return the stringified version of the enum passed in */
-	inline const TCHAR* ToString(EOnlineServerConnectionStatus::Type EnumVal)
-	{
-		switch (EnumVal)
-		{
-			case Normal:
-			{
-				return TEXT("Normal");
-			}
-			case NotConnected:
-			{
-				return TEXT("NotConnected");
-			}
-			case Connected:
-			{
-				return TEXT("Connected");
-			}
-			case ConnectionDropped:
-			{
-				return TEXT("ConnectionDropped");
-			}
-			case NoNetworkConnection:
-			{
-				return TEXT("NoNetworkConnection");
-			}
-			case ServiceUnavailable:
-			{
-				return TEXT("ServiceUnavailable");
-			}
-			case UpdateRequired:
-			{
-				return TEXT("UpdateRequired");
-			}
-			case ServersTooBusy:
-			{
-				return TEXT("ServersTooBusy");
-			}
-			case DuplicateLoginDetected:
-			{
-				return TEXT("DuplicateLoginDetected");
-			}
-			case InvalidUser:
-			{
-				return TEXT("InvalidUser");
-			}
-			case NotAuthorized:
-			{
-				return TEXT("NotAuthorized");
-			}
-			case InvalidSession:
-			{
-				return TEXT("InvalidSession");
-			}
-			default:
-			{
-				return TEXT("Unknown");
-			}
-		}
-		return TEXT("");
-	}
+	ONLINESUBSYSTEM_API const TCHAR* ToString(EOnlineServerConnectionStatus::Type EnumVal);
 };
 
 /** Possible feature privilege access levels */
@@ -197,29 +112,7 @@ namespace EFeaturePrivilegeLevel
 	};
 
 	/** @return the stringified version of the enum passed in */
-	inline const TCHAR* ToString(EFeaturePrivilegeLevel::Type EnumVal)
-	{
-		switch (EnumVal)
-		{
-			case Undefined:
-			{
-				return TEXT("Undefined");
-			}
-			case Disabled:
-			{
-				return TEXT("Disabled");
-			}
-			case EnabledFriendsOnly:
-			{
-				return TEXT("EnabledFriendsOnly");
-			}
-			case Enabled:
-			{
-				return TEXT("Enabled");
-			}
-		}
-		return TEXT("");
-	}
+	ONLINESUBSYSTEM_API const TCHAR* ToString(EFeaturePrivilegeLevel::Type EnumVal);
 }
 
 /** The state of an async task (read friends, read content, write cloud file, etc) request */
@@ -238,29 +131,7 @@ namespace EOnlineAsyncTaskState
 	};  
 
 	/** @return the stringified version of the enum passed in */
-	inline const TCHAR* ToString(EOnlineAsyncTaskState::Type EnumVal)
-	{
-		switch (EnumVal)
-		{
-			case NotStarted:
-			{
-				return TEXT("NotStarted");
-			}
-			case InProgress:
-			{
-				return TEXT("InProgress");
-			}
-			case Done:
-			{
-				return TEXT("Done");
-			}
-			case Failed:
-			{
-				return TEXT("Failed");
-			}
-		}
-		return TEXT("");
-	}
+	ONLINESUBSYSTEM_API const TCHAR* ToString(EOnlineAsyncTaskState::Type EnumVal);
 }
 
 /** The possible friend states for a friend entry */
@@ -279,29 +150,7 @@ namespace EOnlineFriendState
 	};
 
 	/** @return the stringified version of the enum passed in */
-	inline const TCHAR* ToString(EOnlineFriendState::Type EnumVal)
-	{
-		switch (EnumVal)
-		{
-			case Offline:
-			{
-				return TEXT("Offline");
-			}
-			case Online:
-			{
-				return TEXT("Online");
-			}
-			case Away:
-			{
-				return TEXT("Away");
-			}
-			case Busy:
-			{
-				return TEXT("Busy");
-			}
-		}
-		return TEXT("");
-	}
+	ONLINESUBSYSTEM_API const TCHAR* ToString(EOnlineFriendState::Type EnumVal);
 }
 
 /** Leaderboard entry sort types */
@@ -318,25 +167,7 @@ namespace ELeaderboardSort
 	};
 
 	/** @return the stringified version of the enum passed in */
-	inline const TCHAR* ToString(ELeaderboardSort::Type EnumVal)
-	{
-		switch (EnumVal)
-		{
-		case None:
-			{
-				return TEXT("None");
-			}
-		case Ascending:
-			{
-				return TEXT("Ascending");
-			}
-		case Descending:
-			{
-				return TEXT("Descending");
-			}
-		}
-		return TEXT("");
-	}
+	ONLINESUBSYSTEM_API const TCHAR* ToString(ELeaderboardSort::Type EnumVal);
 }
 
 /** Leaderboard display format */
@@ -353,25 +184,7 @@ namespace ELeaderboardFormat
 	};
 
 	/** @return the stringified version of the enum passed in */
-	inline const TCHAR* ToString(ELeaderboardFormat::Type EnumVal)
-	{
-		switch (EnumVal)
-		{
-		case Number:
-			{
-				return TEXT("Number");
-			}
-		case Seconds:
-			{
-				return TEXT("Seconds");
-			}
-		case Milliseconds:
-			{
-				return TEXT("Milliseconds");
-			}
-		}
-		return TEXT("");
-	}
+	ONLINESUBSYSTEM_API const TCHAR* ToString(ELeaderboardFormat::Type EnumVal);
 }
 
 /** How to upload leaderboard score updates */
@@ -386,21 +199,7 @@ namespace ELeaderboardUpdateMethod
 	};
 
 	/** @return the stringified version of the enum passed in */
-	inline const TCHAR* ToString(ELeaderboardUpdateMethod::Type EnumVal)
-	{
-		switch (EnumVal)
-		{
-		case KeepBest:
-			{
-				return TEXT("KeepBest");
-			}
-		case Force:
-			{
-				return TEXT("Force");
-			}
-		}
-		return TEXT("");
-	}
+	ONLINESUBSYSTEM_API const TCHAR* ToString(ELeaderboardUpdateMethod::Type EnumVal);
 }
 
 /** Enum indicating the current state of the online session (in progress, ended, etc.) */
@@ -427,46 +226,7 @@ namespace EOnlineSessionState
 	};
 
 	/** @return the stringified version of the enum passed in */
-	inline const TCHAR* ToString(EOnlineSessionState::Type EnumVal)
-	{
-		switch (EnumVal)
-		{
-
-		case NoSession:
-			{
-				return TEXT("NoSession");
-			}
-		case Creating:
-			{
-				return TEXT("Creating");
-			}
-		case Pending:
-			{
-				return TEXT("Pending");
-			}
-		case Starting:
-			{
-				return TEXT("Starting");
-			}
-		case InProgress:
-			{
-				return TEXT("InProgress");
-			}
-		case Ending:
-			{
-				return TEXT("Ending");
-			}
-		case Ended:
-			{
-				return TEXT("Ended");
-			}
-		case Destroying:
-			{
-				return TEXT("Destroying");
-			}
-		}
-		return TEXT("");
-	}
+	ONLINESUBSYSTEM_API const TCHAR* ToString(EOnlineSessionState::Type EnumVal);
 }
 
 /** The types of advertisement of settings to use */
@@ -485,29 +245,7 @@ namespace EOnlineDataAdvertisementType
 	};
 
 	/** @return the stringified version of the enum passed in */
-	inline const TCHAR* ToString(EOnlineDataAdvertisementType::Type EnumVal)
-	{
-		switch (EnumVal)
-		{
-		case DontAdvertise:
-			{
-				return TEXT("DontAdvertise");
-			}
-		case ViaPingOnly:
-			{
-				return TEXT("ViaPingOnly");
-			}
-		case ViaOnlineService:
-			{
-				return TEXT("OnlineService");
-			}
-		case ViaOnlineServiceAndPing:
-			{
-				return TEXT("OnlineServiceAndPing");
-			}
-		}
-		return TEXT("");
-	}
+	ONLINESUBSYSTEM_API const TCHAR* ToString(EOnlineDataAdvertisementType::Type EnumVal);
 }
 
 /** The types of comparison operations for a given search query */
@@ -527,49 +265,7 @@ namespace EOnlineComparisonOp
 	};
 
 	/** @return the stringified version of the enum passed in */
-	inline const TCHAR* ToString(EOnlineComparisonOp::Type EnumVal)
-	{
-		switch (EnumVal)
-		{
-		case Equals:
-			{
-				return TEXT("Equals");
-			}
-		case NotEquals:
-			{
-				return TEXT("NotEquals");
-			}
-		case GreaterThan:
-			{
-				return TEXT("GreaterThan");
-			}
-		case GreaterThanEquals:
-			{
-				return TEXT("GreaterThanEquals");
-			}
-		case LessThan:
-			{
-				return TEXT("LessThan");
-			}
-		case LessThanEquals:
-			{
-				return TEXT("LessThanEquals");
-			}
-		case Near:
-			{
-				return TEXT("Near");
-			}
-		case In:
-			{
-				return TEXT("In");
-			}
-		case NotIn:
-			{
-				return TEXT("NotIn");
-			}
-		}
-		return TEXT("");
-	}
+	ONLINESUBSYSTEM_API const TCHAR* ToString(EOnlineComparisonOp::Type EnumVal);
 }
 
 /** Return codes for the GetCached functions in the various subsystems. */
@@ -585,21 +281,7 @@ namespace EOnlineCachedResult
 	 * @param EnumVal the enum to convert to a string
 	 * @return the stringified version of the enum passed in
 	 */
-	inline const TCHAR* ToString(EOnlineCachedResult::Type EnumVal)
-	{
-		switch (EnumVal)
-		{
-		case Success:
-			{
-				return TEXT("Success");
-			}
-		case NotFound:
-			{
-				return TEXT("NotFound");
-			}
-		}
-		return TEXT("");
-	}
+	ONLINESUBSYSTEM_API const TCHAR* ToString(EOnlineCachedResult::Type EnumVal);
 }
 
 /** Permissions for who can send invites to a user. */
@@ -615,38 +297,8 @@ enum class EFriendInvitePolicy : uint8
 	InvalidOrMax
 };
 
-inline const TCHAR* LexToString(EFriendInvitePolicy EnumVal)
-{
-	switch (EnumVal)
-	{
-	case EFriendInvitePolicy::Public: return TEXT("PUBLIC");
-	case EFriendInvitePolicy::Friends_of_Friends: return TEXT("FRIENDS_OF_FRIENDS");
-	case EFriendInvitePolicy::Private: return TEXT("PRIVATE");
-	default: return TEXT("Invalid");
-	}
-}
-
-inline void LexFromString(EFriendInvitePolicy& Value, const TCHAR* String)
-{
-	if (FCString::Stricmp(String, TEXT("PUBLIC")) == 0)
-	{
-		Value = EFriendInvitePolicy::Public;			
-	}
-	else if (FCString::Stricmp(String, TEXT("FRIENDS_OF_FRIENDS")) == 0)
-	{
-		Value = EFriendInvitePolicy::Friends_of_Friends;
-	}
-	else if (FCString::Stricmp(String, TEXT("PRIVATE")) == 0)
-	{
-		Value = EFriendInvitePolicy::Private;
-	}
-	else
-	{
-		Value = EFriendInvitePolicy::InvalidOrMax;
-	}
-}
-
-
+ONLINESUBSYSTEM_API const TCHAR* LexToString(EFriendInvitePolicy EnumVal);
+ONLINESUBSYSTEM_API void LexFromString(EFriendInvitePolicy& Value, const TCHAR* String);
 
 /*
  *	Base class for anything meant to be opaque so that the data can be passed around 
@@ -674,11 +326,7 @@ protected:
 		return *this;
 	}
 
-	virtual bool Compare(const IOnlinePlatformData& Other) const
-	{
-		return (GetSize() == Other.GetSize()) &&
-			(FMemory::Memcmp(GetBytes(), Other.GetBytes(), GetSize()) == 0);
-	}
+	virtual ONLINESUBSYSTEM_API bool Compare(const IOnlinePlatformData& Other) const;
 
 public:
 
@@ -1283,37 +931,7 @@ namespace EInviteStatus
 	/** 
 	 * @return the stringified version of the enum passed in 
 	 */
-	inline const TCHAR* ToString(EInviteStatus::Type EnumVal)
-	{
-		switch (EnumVal)
-		{
-			case Unknown:
-			{
-				return TEXT("Unknown");
-			}
-			case Accepted:
-			{
-				return TEXT("Accepted");
-			}
-			case PendingInbound:
-			{
-				return TEXT("PendingInbound");
-			}
-			case PendingOutbound:
-			{
-				return TEXT("PendingOutbound");
-			}
-			case Blocked:
-			{
-				return TEXT("Blocked");
-			}
-			case Suggested:
-			{
-				return TEXT("Suggested");
-			}
-		}
-		return TEXT("");
-	}
+	ONLINESUBSYSTEM_API const TCHAR* ToString(EInviteStatus::Type EnumVal);
 };
 
 /**
@@ -1405,57 +1023,7 @@ enum class EOnlineSharingCategory : uint32
 
 ENUM_CLASS_FLAGS(EOnlineSharingCategory);
 
-inline const TCHAR* ToString(EOnlineSharingCategory CategoryType)
-{
-	switch (CategoryType)
-	{
-		case EOnlineSharingCategory::None:
-		{
-			return TEXT("Category undefined");
-		}
-		case EOnlineSharingCategory::ReadPosts:
-		{
-			return TEXT("ReadPosts");
-		}
-		case EOnlineSharingCategory::Friends:
-		{
-			return TEXT("Friends");
-		}
-		case EOnlineSharingCategory::Mailbox:
-		{
-			return TEXT("Mailbox");
-		}
-		case EOnlineSharingCategory::OnlineStatus:
-		{
-			return TEXT("Online Status");
-		}
-		case EOnlineSharingCategory::ProfileInfo:
-		{
-			return TEXT("Profile Information");
-		}
-		case EOnlineSharingCategory::LocationInfo:
-		{
-			return TEXT("Location Information");
-		}
-		case EOnlineSharingCategory::SubmitPosts:
-		{
-			return TEXT("SubmitPosts");
-		}
-		case EOnlineSharingCategory::ManageFriends:
-		{
-			return TEXT("ManageFriends");
-		}
-		case EOnlineSharingCategory::AccountAdmin:
-		{
-			return TEXT("Account Admin");
-		}
-		case EOnlineSharingCategory::Events:
-		{
-			return TEXT("Events");
-		}
-	}
-	return TEXT("");
-}
+ONLINESUBSYSTEM_API const TCHAR* ToString(EOnlineSharingCategory CategoryType);
 
 /** Privacy permissions used for Online Status updates */
 enum class EOnlineStatusUpdatePrivacy : uint8
@@ -1468,18 +1036,7 @@ enum class EOnlineStatusUpdatePrivacy : uint8
 	Everyone,
 };
 
-inline const TCHAR* ToString(EOnlineStatusUpdatePrivacy PrivacyType)
-{
-	switch (PrivacyType)
-	{
-		case EOnlineStatusUpdatePrivacy::OnlyMe:
-			return TEXT("Only Me");
-		case EOnlineStatusUpdatePrivacy::OnlyFriends:
-			return TEXT("Only Friends");
-		case EOnlineStatusUpdatePrivacy::Everyone:
-			return TEXT("Everyone");
-	}
-}
+ONLINESUBSYSTEM_API const TCHAR* ToString(EOnlineStatusUpdatePrivacy PrivacyType);
 
 /**
  * unique identifier for notification transports
@@ -1508,21 +1065,18 @@ protected:
 		return *this;
 	}
 
-	friend inline uint32 GetTypeHash(const FOnlinePartyId& Value)
-	{
-		return CityHash32(reinterpret_cast<const char*>(Value.GetBytes()), Value.GetSize());
-	}
-
 public:
 	virtual ~FOnlinePartyId() {}
 };
 
+ONLINESUBSYSTEM_API uint32 GetTypeHash(const FOnlinePartyId& Value);
+
 template <typename ValueType>
 struct TOnlinePartyIdMapKeyFuncs : public TDefaultMapKeyFuncs<TSharedRef<const FOnlinePartyId>, ValueType, false>
 {
-	static FORCEINLINE TSharedRef<const FOnlinePartyId>	GetSetKey(TPair<TSharedRef<const FOnlinePartyId>, ValueType> const& Element) { return Element.Key; }
-	static FORCEINLINE uint32							GetKeyHash(TSharedRef<const FOnlinePartyId> const& Key) {	return GetTypeHash(*Key); }
-	static FORCEINLINE bool								Matches(TSharedRef<const FOnlinePartyId> const& A, TSharedRef<const FOnlinePartyId> const& B) { return (A == B) || (*A == *B); }
+	static TSharedRef<const FOnlinePartyId>	GetSetKey(TPair<TSharedRef<const FOnlinePartyId>, ValueType> const& Element) { return Element.Key; }
+	static uint32							GetKeyHash(TSharedRef<const FOnlinePartyId> const& Key) {	return GetTypeHash(*Key); }
+	static bool								Matches(TSharedRef<const FOnlinePartyId> const& A, TSharedRef<const FOnlinePartyId> const& B) { return (A == B) || (*A == *B); }
 };
 
 template <typename ValueType>
@@ -1582,32 +1136,7 @@ public:
  * @param InEntries array of strings in the format (Key=Value)
  * @param OutPairs the split key/value pairs
  */
-inline void ParseOnlineSubsystemConfigPairs(TArrayView<const FString> InEntries, TArray<TPair<FString, FString>>& OutPairs)
-{
-	OutPairs.Reserve(InEntries.Num());
-	// Takes on the pattern "(Key=Value)"
-	for (const FString& Entry : InEntries)
-	{
-		FString TrimmedConfigEntry = Entry.TrimStartAndEnd();
-		FString KeyString;
-		FString ValueString;
-
-		if (TrimmedConfigEntry.Left(1) == TEXT("("))
-		{
-			TrimmedConfigEntry.RightChopInline(1, false);
-		}
-		if (TrimmedConfigEntry.Right(1) == TEXT(")"))
-		{
-			TrimmedConfigEntry.LeftChopInline(1, false);
-		}
-		if (TrimmedConfigEntry.Split(TEXT("="), &KeyString, &ValueString, ESearchCase::CaseSensitive))
-		{
-			KeyString.TrimStartAndEndInline();
-			ValueString.TrimStartAndEndInline();
-		}
-		OutPairs.Emplace(MoveTemp(KeyString), MoveTemp(ValueString));
-	}
-}
+ONLINESUBSYSTEM_API void ParseOnlineSubsystemConfigPairs(TArrayView<const FString> InEntries, TArray<TPair<FString, FString>>& OutPairs);
 
 #if UE_ENABLE_INCLUDE_ORDER_DEPRECATED_IN_5_2
 #include "CoreMinimal.h"
