Implement a content restriction interface that supported platforms can leverage

diff --git a/Engine/Plugins/Online/OnlineSubsystem/Source/Private/OnlineSubsystemImpl.cpp b/Engine/Plugins/Online/OnlineSubsystem/Source/Private/OnlineSubsystemImpl.cpp
--- a/Engine/Plugins/Online/OnlineSubsystem/Source/Private/OnlineSubsystemImpl.cpp
+++ b/Engine/Plugins/Online/OnlineSubsystem/Source/Private/OnlineSubsystemImpl.cpp
@@ -812,3 +812,8 @@ IOnlineTournamentPtr FOnlineSubsystemImpl::GetTournamentInterface() const
 {
 	return nullptr;
 }
+
+IOnlineContentAgeRestrictionPtr FOnlineSubsystemImpl::GetOnlineContentAgeRestrictionInterface() const
+{
+	return nullptr;
+}
diff --git a/Engine/Plugins/Online/OnlineSubsystem/Source/Public/Interfaces/OnlineContentRestrictionInterface.h b/Engine/Plugins/Online/OnlineSubsystem/Source/Public/Interfaces/OnlineContentRestrictionInterface.h
new file mode 100644
--- /dev/null
+++ b/Engine/Plugins/Online/OnlineSubsystem/Source/Public/Interfaces/OnlineContentRestrictionInterface.h
@@ -0,0 +1,36 @@
+// Copyright Epic Games, Inc. All Rights Reserved.
+
+#pragma once
+
+#include "OnlineDelegateMacros.h"
+#include "OnlineSubsystemTypes.h"
+
+class FUniqueNetId;
+struct FOnlineError;
+
+/**
+ * Delegate fired when content age restriction check is completed
+ * 
+ * @param LocalUserId the id of the player that this callback is intended for
+ * @param bAgeRestricted the restriction result (is or is not restricted)
+ * @param Result whether or not the process completed successfully 
+ */
+DECLARE_DELEGATE_ThreeParams(FOnQueryContentAgeRestictionComplete, const FUniqueNetId& /* LocalUserId */, const bool /* bAgeRestricted */, const FOnlineError& /* Result */);
+
+/**
+ *	IOnlineContentAgeRestriction - Interface class for managing a user's content age restriction status
+ */
+class IOnlineContentAgeRestriction
+{
+public:
+	virtual ~IOnlineContentAgeRestriction() = default;
+
+	/**
+	 * Query the content age restriction status for an account
+	 * 
+	 * @param LocalUserId - Id of the player to check the content age restriction status for
+	 * @param RatingBoards - Array of rating boards defined to query status with
+	 * @param CompletionDelegate - Completion delegate called when the QueryContentAgeRestriction call is complete
+	 */
+	virtual void QueryContentAgeRestriction(const FUniqueNetId& LocalUserId, const TArray<FString>& RatingBoards, FOnQueryContentAgeRestictionComplete CompletionDelegate) = 0;
+};
\ No newline at end of file
diff --git a/Engine/Plugins/Online/OnlineSubsystem/Source/Public/OnlineSubsystem.h b/Engine/Plugins/Online/OnlineSubsystem/Source/Public/OnlineSubsystem.h
--- a/Engine/Plugins/Online/OnlineSubsystem/Source/Public/OnlineSubsystem.h
+++ b/Engine/Plugins/Online/OnlineSubsystem/Source/Public/OnlineSubsystem.h
@@ -40,6 +40,7 @@ class IOnlineStats;
 class IOnlineGameActivity;
 class IOnlineGameItemStats;
 class IOnlineGameMatches;
+class IOnlineContentAgeRestriction;
 
 LLM_DECLARE_TAG_API(OnlineSubsystem, ONLINESUBSYSTEM_API);
 
@@ -110,6 +111,7 @@ typedef TSharedPtr<class IOnlineStats, ESPMode::ThreadSafe> IOnlineStatsPtr;
 typedef TSharedPtr<IOnlineGameActivity, ESPMode::ThreadSafe> IOnlineGameActivityPtr;
 typedef TSharedPtr<IOnlineGameItemStats, ESPMode::ThreadSafe> IOnlineGameItemStatsPtr;
 typedef TSharedPtr<IOnlineGameMatches, ESPMode::ThreadSafe> IOnlineGameMatchesPtr;
+typedef TSharedPtr<IOnlineContentAgeRestriction, ESPMode::ThreadSafe> IOnlineContentAgeRestrictionPtr;
 
 /**
  * Called when the connection state as reported by the online platform changes
@@ -457,6 +459,12 @@ public:
 	 */
 	virtual IOnlineGameMatchesPtr GetGameMatchesInterface() const = 0;
 
+	/**
+	 * Get interface for accessing the online content age restriction
+	 * @return Interface pointer for online content age restriction
+	 */
+	 virtual IOnlineContentAgeRestrictionPtr GetOnlineContentAgeRestrictionInterface() const = 0;
+
 	/**
 	 * Get the notification handler instance for this subsystem
 	 * @return Pointer for the appropriate notification handler
diff --git a/Engine/Plugins/Online/OnlineSubsystem/Source/Public/OnlineSubsystemImpl.h b/Engine/Plugins/Online/OnlineSubsystem/Source/Public/OnlineSubsystemImpl.h
--- a/Engine/Plugins/Online/OnlineSubsystem/Source/Public/OnlineSubsystemImpl.h
+++ b/Engine/Plugins/Online/OnlineSubsystem/Source/Public/OnlineSubsystemImpl.h
@@ -113,6 +113,7 @@ public:
 	virtual IOnlineGameMatchesPtr GetGameMatchesInterface() const override;
 	virtual IOnlineTurnBasedPtr GetTurnBasedInterface() const override;
 	virtual IOnlineTournamentPtr GetTournamentInterface() const override;
+	virtual IOnlineContentAgeRestrictionPtr GetOnlineContentAgeRestrictionInterface() const override;
 	virtual void PreUnload() override;
 	virtual bool Shutdown() override;
 	virtual bool IsServer() const override;
