Add ShowResolveUI param to GetUserPrivilege method of identity interface.

#tests Tried on xsx, now it works.

diff --git a/Engine/Plugins/Online/Android/OnlineSubsystemGooglePlay/Source/Private/OnlineIdentityInterfaceGooglePlay.cpp b/Engine/Plugins/Online/Android/OnlineSubsystemGooglePlay/Source/Private/OnlineIdentityInterfaceGooglePlay.cpp
--- a/Engine/Plugins/Online/Android/OnlineSubsystemGooglePlay/Source/Private/OnlineIdentityInterfaceGooglePlay.cpp
+++ b/Engine/Plugins/Online/Android/OnlineSubsystemGooglePlay/Source/Private/OnlineIdentityInterfaceGooglePlay.cpp
@@ -176,7 +176,7 @@ void FOnlineIdentityGooglePlay::OnLoginCompleted(const int playerID, const gpg::
 	PendingConnectRequest.IsConnectionPending = false;
 }
 
-void FOnlineIdentityGooglePlay::GetUserPrivilege(const FUniqueNetId& UserId, EUserPrivileges::Type Privilege, const FOnGetUserPrivilegeCompleteDelegate& Delegate)
+void FOnlineIdentityGooglePlay::GetUserPrivilege(const FUniqueNetId& UserId, EUserPrivileges::Type Privilege, const FOnGetUserPrivilegeCompleteDelegate& Delegate, EShowPrivilegeResolveUI ShowResolveUI)
 {
 	Delegate.ExecuteIfBound(UserId, Privilege, (uint32)EPrivilegeResults::NoFailures);
 }
diff --git a/Engine/Plugins/Online/Android/OnlineSubsystemGooglePlay/Source/Private/OnlineIdentityInterfaceGooglePlay.h b/Engine/Plugins/Online/Android/OnlineSubsystemGooglePlay/Source/Private/OnlineIdentityInterfaceGooglePlay.h
--- a/Engine/Plugins/Online/Android/OnlineSubsystemGooglePlay/Source/Private/OnlineIdentityInterfaceGooglePlay.h
+++ b/Engine/Plugins/Online/Android/OnlineSubsystemGooglePlay/Source/Private/OnlineIdentityInterfaceGooglePlay.h
@@ -76,7 +76,7 @@ public:
 	virtual FString GetPlayerNickname(const FUniqueNetId& UserId) const override;
 	virtual FString GetAuthToken(int32 LocalUserNum) const override;
 	virtual void RevokeAuthToken(const FUniqueNetId& UserId, const FOnRevokeAuthTokenCompleteDelegate& Delegate) override;
-	virtual void GetUserPrivilege(const FUniqueNetId& UserId, EUserPrivileges::Type Privilege, const FOnGetUserPrivilegeCompleteDelegate& Delegate) override;
+	virtual void GetUserPrivilege(const FUniqueNetId& UserId, EUserPrivileges::Type Privilege, const FOnGetUserPrivilegeCompleteDelegate& Delegate, EShowPrivilegeResolveUI ShowResolveUI=EShowPrivilegeResolveUI::Default) override;
 	virtual FPlatformUserId GetPlatformUserIdFromUniqueNetId(const FUniqueNetId& NetId) const override;
 	virtual FString GetAuthType() const override;
 	//~ End IOnlineIdentity interface
diff --git a/Engine/Plugins/Online/IOS/OnlineSubsystemIOS/Source/Private/OnlineIdentityInterfaceIOS.cpp b/Engine/Plugins/Online/IOS/OnlineSubsystemIOS/Source/Private/OnlineIdentityInterfaceIOS.cpp
--- a/Engine/Plugins/Online/IOS/OnlineSubsystemIOS/Source/Private/OnlineIdentityInterfaceIOS.cpp
+++ b/Engine/Plugins/Online/IOS/OnlineSubsystemIOS/Source/Private/OnlineIdentityInterfaceIOS.cpp
@@ -306,7 +306,7 @@ void FOnlineIdentityIOS::RevokeAuthToken(const FUniqueNetId& UserId, const FOnRe
 	});
 }
 
-void FOnlineIdentityIOS::GetUserPrivilege(const FUniqueNetId& UserId, EUserPrivileges::Type Privilege, const FOnGetUserPrivilegeCompleteDelegate& Delegate)
+void FOnlineIdentityIOS::GetUserPrivilege(const FUniqueNetId& UserId, EUserPrivileges::Type Privilege, const FOnGetUserPrivilegeCompleteDelegate& Delegate, EShowPrivilegeResolveUI ShowResolveUI)
 {
 	if (Privilege == EUserPrivileges::CanPlayOnline)
 	{
diff --git a/Engine/Plugins/Online/IOS/OnlineSubsystemIOS/Source/Private/OnlineIdentityInterfaceIOS.h b/Engine/Plugins/Online/IOS/OnlineSubsystemIOS/Source/Private/OnlineIdentityInterfaceIOS.h
--- a/Engine/Plugins/Online/IOS/OnlineSubsystemIOS/Source/Private/OnlineIdentityInterfaceIOS.h
+++ b/Engine/Plugins/Online/IOS/OnlineSubsystemIOS/Source/Private/OnlineIdentityInterfaceIOS.h
@@ -53,7 +53,7 @@ public:
 	virtual FString GetPlayerNickname(const FUniqueNetId& UserId) const override;
 	virtual FString GetAuthToken(int32 LocalUserNum) const override;
 	virtual void RevokeAuthToken(const FUniqueNetId& UserId, const FOnRevokeAuthTokenCompleteDelegate& Delegate) override;
-	virtual void GetUserPrivilege(const FUniqueNetId& UserId, EUserPrivileges::Type Privilege, const FOnGetUserPrivilegeCompleteDelegate& Delegate) override;
+	virtual void GetUserPrivilege(const FUniqueNetId& UserId, EUserPrivileges::Type Privilege, const FOnGetUserPrivilegeCompleteDelegate& Delegate, EShowPrivilegeResolveUI ShowResolveUI=EShowPrivilegeResolveUI::Default) override;
 	virtual FPlatformUserId GetPlatformUserIdFromUniqueNetId(const FUniqueNetId& UniqueNetId) const override;
 	virtual FString GetAuthType() const override;
 	//~ End IOnlineIdentity Interface
diff --git a/Engine/Plugins/Online/OnlineServicesOSSAdapter/Source/Private/Online/PrivilegesOSSAdapter.cpp b/Engine/Plugins/Online/OnlineServicesOSSAdapter/Source/Private/Online/PrivilegesOSSAdapter.cpp
--- a/Engine/Plugins/Online/OnlineServicesOSSAdapter/Source/Private/Online/PrivilegesOSSAdapter.cpp
+++ b/Engine/Plugins/Online/OnlineServicesOSSAdapter/Source/Private/Online/PrivilegesOSSAdapter.cpp
@@ -63,7 +63,7 @@ TOnlineAsyncOpHandle<FQueryUserPrivilege> FPrivilegesOSSAdapter::QueryUserPrivil
 
 								Promise->SetValue();
 							}
-						}));
+						}), EShowPrivilegeResolveUI::Default);
 
 				return Promise->GetFuture();
 			})
diff --git a/Engine/Plugins/Online/OnlineSubsystem/Source/Public/Interfaces/OnlineIdentityInterface.h b/Engine/Plugins/Online/OnlineSubsystem/Source/Public/Interfaces/OnlineIdentityInterface.h
--- a/Engine/Plugins/Online/OnlineSubsystem/Source/Public/Interfaces/OnlineIdentityInterface.h
+++ b/Engine/Plugins/Online/OnlineSubsystem/Source/Public/Interfaces/OnlineIdentityInterface.h
@@ -84,6 +84,17 @@ namespace EUserPrivileges
 	};
 }
 
+/** Whether to show UI to resolve the privilege if there is no access when GetUserPrivilege */
+enum class EShowPrivilegeResolveUI : uint8
+{
+	/** Use default behavior, which preserved the behaviour prior to the introduction of this. */
+	Default,
+	/** Show UI to resolve privilege */
+	Show,
+	/** Not show UI to resolve privilege */
+	NotShow
+};
+
 /**
  * This struct will be broadcast with the FOnControllerPairingChanged event. This event is not broadcast on all OnlineIdentity platforms.
  * There will be a struct for the previous user on the controller and for the new user on the controller.
@@ -488,8 +499,9 @@ public:
 	 * @param LocalUserId the unique id of the user to query
 	 * @param Privilege the privilege you want to know about
 	 * @param Delegate delegate to execute when the async task completes
+	 * @param ShowResolveUI whether to show UI to resolve if user don't have the privilege
 	 */
-	virtual void GetUserPrivilege(const FUniqueNetId& LocalUserId, EUserPrivileges::Type Privilege, const FOnGetUserPrivilegeCompleteDelegate& Delegate) = 0;
+	virtual void GetUserPrivilege(const FUniqueNetId& LocalUserId, EUserPrivileges::Type Privilege, const FOnGetUserPrivilegeCompleteDelegate& Delegate, EShowPrivilegeResolveUI ShowResolveUI = EShowPrivilegeResolveUI::Default) = 0;
 
 	/**
 	 * Converts from an online unique id to a PlatformUserId used by application code
diff --git a/Engine/Plugins/Online/OnlineSubsystemAmazon/Source/Private/OnlineIdentityAmazon.cpp b/Engine/Plugins/Online/OnlineSubsystemAmazon/Source/Private/OnlineIdentityAmazon.cpp
--- a/Engine/Plugins/Online/OnlineSubsystemAmazon/Source/Private/OnlineIdentityAmazon.cpp
+++ b/Engine/Plugins/Online/OnlineSubsystemAmazon/Source/Private/OnlineIdentityAmazon.cpp
@@ -355,7 +355,7 @@ void FOnlineIdentityAmazon::RevokeAuthToken(const FUniqueNetId& UserId, const FO
 	});
 }
 
-void FOnlineIdentityAmazon::GetUserPrivilege(const FUniqueNetId& UserId, EUserPrivileges::Type Privilege, const FOnGetUserPrivilegeCompleteDelegate& Delegate)
+void FOnlineIdentityAmazon::GetUserPrivilege(const FUniqueNetId& UserId, EUserPrivileges::Type Privilege, const FOnGetUserPrivilegeCompleteDelegate& Delegate, EShowPrivilegeResolveUI ShowResolveUI)
 {
 	Delegate.ExecuteIfBound(UserId, Privilege, (uint32)EPrivilegeResults::NoFailures);
 }
diff --git a/Engine/Plugins/Online/OnlineSubsystemAmazon/Source/Private/OnlineIdentityAmazon.h b/Engine/Plugins/Online/OnlineSubsystemAmazon/Source/Private/OnlineIdentityAmazon.h
--- a/Engine/Plugins/Online/OnlineSubsystemAmazon/Source/Private/OnlineIdentityAmazon.h
+++ b/Engine/Plugins/Online/OnlineSubsystemAmazon/Source/Private/OnlineIdentityAmazon.h
@@ -126,7 +126,7 @@ public:
 	virtual FString GetPlayerNickname(const FUniqueNetId& UserId) const override;
 	virtual FString GetAuthToken(int32 LocalUserNum) const override;
 	virtual void RevokeAuthToken(const FUniqueNetId& UserId, const FOnRevokeAuthTokenCompleteDelegate& Delegate) override;
-	virtual void GetUserPrivilege(const FUniqueNetId& UserId, EUserPrivileges::Type Privilege, const FOnGetUserPrivilegeCompleteDelegate& Delegate) override;
+	virtual void GetUserPrivilege(const FUniqueNetId& UserId, EUserPrivileges::Type Privilege, const FOnGetUserPrivilegeCompleteDelegate& Delegate, EShowPrivilegeResolveUI ShowResolveUI=EShowPrivilegeResolveUI::Default) override;
 	virtual FPlatformUserId GetPlatformUserIdFromUniqueNetId(const FUniqueNetId& UniqueNetId) const override;
 	virtual FString GetAuthType() const override;
 
diff --git a/Engine/Plugins/Online/OnlineSubsystemApple/Source/Private/OnlineIdentityInterfaceApple.cpp b/Engine/Plugins/Online/OnlineSubsystemApple/Source/Private/OnlineIdentityInterfaceApple.cpp
--- a/Engine/Plugins/Online/OnlineSubsystemApple/Source/Private/OnlineIdentityInterfaceApple.cpp
+++ b/Engine/Plugins/Online/OnlineSubsystemApple/Source/Private/OnlineIdentityInterfaceApple.cpp
@@ -379,7 +379,7 @@ void FOnlineIdentityApple::RevokeAuthToken(const FUniqueNetId& UserId, const FOn
 	});
 }
 
-void FOnlineIdentityApple::GetUserPrivilege(const FUniqueNetId& UserId, EUserPrivileges::Type Privilege, const FOnGetUserPrivilegeCompleteDelegate& Delegate)
+void FOnlineIdentityApple::GetUserPrivilege(const FUniqueNetId& UserId, EUserPrivileges::Type Privilege, const FOnGetUserPrivilegeCompleteDelegate& Delegate, EShowPrivilegeResolveUI ShowResolveUI)
 {
 	FUniqueNetIdRef UserIdRef(UserId.AsShared());
 	Subsystem->ExecuteNextTick([UserIdRef, Privilege, Delegate]()
diff --git a/Engine/Plugins/Online/OnlineSubsystemApple/Source/Private/OnlineIdentityInterfaceApple.h b/Engine/Plugins/Online/OnlineSubsystemApple/Source/Private/OnlineIdentityInterfaceApple.h
--- a/Engine/Plugins/Online/OnlineSubsystemApple/Source/Private/OnlineIdentityInterfaceApple.h
+++ b/Engine/Plugins/Online/OnlineSubsystemApple/Source/Private/OnlineIdentityInterfaceApple.h
@@ -98,7 +98,7 @@ public:
 	virtual FString GetPlayerNickname(const FUniqueNetId& UserId) const override;
 	virtual FString GetAuthToken(int32 LocalUserNum) const override;
 	virtual void RevokeAuthToken(const FUniqueNetId& UserId, const FOnRevokeAuthTokenCompleteDelegate& Delegate) override;
-	virtual void GetUserPrivilege(const FUniqueNetId& UserId, EUserPrivileges::Type Privilege, const FOnGetUserPrivilegeCompleteDelegate& Delegate) override;
+	virtual void GetUserPrivilege(const FUniqueNetId& UserId, EUserPrivileges::Type Privilege, const FOnGetUserPrivilegeCompleteDelegate& Delegate, EShowPrivilegeResolveUI ShowResolveUI=EShowPrivilegeResolveUI::Default) override;
 	virtual FPlatformUserId GetPlatformUserIdFromUniqueNetId(const FUniqueNetId& UniqueNetId) const override;
 	virtual FString GetAuthType() const override;
 	//~ End IOnlineIdentity Interface
diff --git a/Engine/Plugins/Online/OnlineSubsystemEOS/Source/OnlineSubsystemEOS/Private/UserManagerEOS.cpp b/Engine/Plugins/Online/OnlineSubsystemEOS/Source/OnlineSubsystemEOS/Private/UserManagerEOS.cpp
--- a/Engine/Plugins/Online/OnlineSubsystemEOS/Source/OnlineSubsystemEOS/Private/UserManagerEOS.cpp
+++ b/Engine/Plugins/Online/OnlineSubsystemEOS/Source/OnlineSubsystemEOS/Private/UserManagerEOS.cpp
@@ -1681,7 +1681,7 @@ void FUserManagerEOS::GetLinkedAccountAuthToken(int32 LocalUserNum, const FStrin
 	Delegate.ExecuteIfBound(LocalUserNum, ExternalToken.IsValid(), ExternalToken);
 }
 
-void FUserManagerEOS::GetUserPrivilege(const FUniqueNetId& UserId, EUserPrivileges::Type Privilege, const FOnGetUserPrivilegeCompleteDelegate& Delegate)
+void FUserManagerEOS::GetUserPrivilege(const FUniqueNetId& UserId, EUserPrivileges::Type Privilege, const FOnGetUserPrivilegeCompleteDelegate& Delegate, EShowPrivilegeResolveUI ShowResolveUI)
 {
 	Delegate.ExecuteIfBound(UserId, Privilege, (uint32)EPrivilegeResults::NoFailures);
 }
diff --git a/Engine/Plugins/Online/OnlineSubsystemEOS/Source/OnlineSubsystemEOS/Private/UserManagerEOS.h b/Engine/Plugins/Online/OnlineSubsystemEOS/Source/OnlineSubsystemEOS/Private/UserManagerEOS.h
--- a/Engine/Plugins/Online/OnlineSubsystemEOS/Source/OnlineSubsystemEOS/Private/UserManagerEOS.h
+++ b/Engine/Plugins/Online/OnlineSubsystemEOS/Source/OnlineSubsystemEOS/Private/UserManagerEOS.h
@@ -281,7 +281,7 @@ public:
 	virtual FString GetPlayerNickname(int32 LocalUserNum) const override;
 	virtual FString GetPlayerNickname(const FUniqueNetId& UserId) const override;
 	virtual FString GetAuthToken(int32 LocalUserNum) const override;
-	virtual void GetUserPrivilege(const FUniqueNetId& UserId, EUserPrivileges::Type Privilege, const FOnGetUserPrivilegeCompleteDelegate& Delegate) override;
+	virtual void GetUserPrivilege(const FUniqueNetId& UserId, EUserPrivileges::Type Privilege, const FOnGetUserPrivilegeCompleteDelegate& Delegate, EShowPrivilegeResolveUI ShowResolveUI = EShowPrivilegeResolveUI::Default) override;
 	virtual FString GetAuthType() const override;
 	virtual void RevokeAuthToken(const FUniqueNetId& LocalUserId, const FOnRevokeAuthTokenCompleteDelegate& Delegate) override;
 	virtual FPlatformUserId GetPlatformUserIdFromUniqueNetId(const FUniqueNetId& UniqueNetId) const override;
diff --git a/Engine/Plugins/Online/OnlineSubsystemEOS/Source/OnlineSubsystemEOSPlus/Private/OnlineUserEOSPlus.cpp b/Engine/Plugins/Online/OnlineSubsystemEOS/Source/OnlineSubsystemEOSPlus/Private/OnlineUserEOSPlus.cpp
--- a/Engine/Plugins/Online/OnlineSubsystemEOS/Source/OnlineSubsystemEOSPlus/Private/OnlineUserEOSPlus.cpp
+++ b/Engine/Plugins/Online/OnlineSubsystemEOS/Source/OnlineSubsystemEOSPlus/Private/OnlineUserEOSPlus.cpp
@@ -810,7 +810,7 @@ FString FOnlineUserEOSPlus::GetAuthToken(int32 LocalUserNum) const
 	return BaseIdentityInterface->GetAuthToken(LocalUserNum);
 }
 
-void FOnlineUserEOSPlus::GetUserPrivilege(const FUniqueNetId& UserId, EUserPrivileges::Type Privilege, const FOnGetUserPrivilegeCompleteDelegate& Delegate)
+void FOnlineUserEOSPlus::GetUserPrivilege(const FUniqueNetId& UserId, EUserPrivileges::Type Privilege, const FOnGetUserPrivilegeCompleteDelegate& Delegate, EShowPrivilegeResolveUI ShowResolveUI)
 {
 	FUniqueNetIdEOSPlusPtr NetIdPlus = GetNetIdPlus(UserId.ToString());
 	if (NetIdPlus.IsValid())
@@ -826,7 +826,7 @@ void FOnlineUserEOSPlus::GetUserPrivilege(const FUniqueNetId& UserId, EUserPrivi
 			OriginalDelegate.ExecuteIfBound(*NetIdPlus, Privilege, PrivilegeResults);
 		});
 
-		BaseIdentityInterface->GetUserPrivilege(*NetIdPlus->GetBaseNetId(), Privilege, IntermediateDelegate);
+		BaseIdentityInterface->GetUserPrivilege(*NetIdPlus->GetBaseNetId(), Privilege, IntermediateDelegate, ShowResolveUI);
 	}
 	else
 	{
diff --git a/Engine/Plugins/Online/OnlineSubsystemEOS/Source/OnlineSubsystemEOSPlus/Private/OnlineUserEOSPlus.h b/Engine/Plugins/Online/OnlineSubsystemEOS/Source/OnlineSubsystemEOSPlus/Private/OnlineUserEOSPlus.h
--- a/Engine/Plugins/Online/OnlineSubsystemEOS/Source/OnlineSubsystemEOSPlus/Private/OnlineUserEOSPlus.h
+++ b/Engine/Plugins/Online/OnlineSubsystemEOS/Source/OnlineSubsystemEOSPlus/Private/OnlineUserEOSPlus.h
@@ -385,7 +385,7 @@ public:
 	virtual FString GetPlayerNickname(int32 LocalUserNum) const override;
 	virtual FString GetPlayerNickname(const FUniqueNetId& UserId) const override;
 	virtual FString GetAuthToken(int32 LocalUserNum) const override;
-	virtual void GetUserPrivilege(const FUniqueNetId& UserId, EUserPrivileges::Type Privilege, const FOnGetUserPrivilegeCompleteDelegate& Delegate) override;
+	virtual void GetUserPrivilege(const FUniqueNetId& UserId, EUserPrivileges::Type Privilege, const FOnGetUserPrivilegeCompleteDelegate& Delegate, EShowPrivilegeResolveUI ShowResolveUI=EShowPrivilegeResolveUI::Default) override;
 	virtual FString GetAuthType() const override;
 	virtual void RevokeAuthToken(const FUniqueNetId& LocalUserId, const FOnRevokeAuthTokenCompleteDelegate& Delegate) override;
 	virtual FPlatformUserId GetPlatformUserIdFromUniqueNetId(const FUniqueNetId& UniqueNetId) const override;
diff --git a/Engine/Plugins/Online/OnlineSubsystemFacebook/Source/Private/OnlineIdentityFacebookCommon.cpp b/Engine/Plugins/Online/OnlineSubsystemFacebook/Source/Private/OnlineIdentityFacebookCommon.cpp
--- a/Engine/Plugins/Online/OnlineSubsystemFacebook/Source/Private/OnlineIdentityFacebookCommon.cpp
+++ b/Engine/Plugins/Online/OnlineSubsystemFacebook/Source/Private/OnlineIdentityFacebookCommon.cpp
@@ -287,7 +287,7 @@ void FOnlineIdentityFacebookCommon::RevokeAuthToken(const FUniqueNetId& UserId,
 	});
 }
 
-void FOnlineIdentityFacebookCommon::GetUserPrivilege(const FUniqueNetId& UserId, EUserPrivileges::Type Privilege, const FOnGetUserPrivilegeCompleteDelegate& Delegate)
+void FOnlineIdentityFacebookCommon::GetUserPrivilege(const FUniqueNetId& UserId, EUserPrivileges::Type Privilege, const FOnGetUserPrivilegeCompleteDelegate& Delegate, EShowPrivilegeResolveUI ShowResolveUI)
 {
 	Delegate.ExecuteIfBound(UserId, Privilege, (uint32)EPrivilegeResults::NoFailures);
 }	
diff --git a/Engine/Plugins/Online/OnlineSubsystemFacebook/Source/Private/OnlineIdentityFacebookCommon.h b/Engine/Plugins/Online/OnlineSubsystemFacebook/Source/Private/OnlineIdentityFacebookCommon.h
--- a/Engine/Plugins/Online/OnlineSubsystemFacebook/Source/Private/OnlineIdentityFacebookCommon.h
+++ b/Engine/Plugins/Online/OnlineSubsystemFacebook/Source/Private/OnlineIdentityFacebookCommon.h
@@ -58,7 +58,7 @@ public:
 	virtual FString GetPlayerNickname(const FUniqueNetId& UserId) const override;
 	virtual FString GetAuthToken(int32 LocalUserNum) const override;
 	virtual void RevokeAuthToken(const FUniqueNetId& UserId, const FOnRevokeAuthTokenCompleteDelegate& Delegate) override;
-	virtual void GetUserPrivilege(const FUniqueNetId& UserId, EUserPrivileges::Type Privilege, const FOnGetUserPrivilegeCompleteDelegate& Delegate) override;
+	virtual void GetUserPrivilege(const FUniqueNetId& UserId, EUserPrivileges::Type Privilege, const FOnGetUserPrivilegeCompleteDelegate& Delegate, EShowPrivilegeResolveUI ShowResolveUI=EShowPrivilegeResolveUI::Default) override;
 	virtual FPlatformUserId GetPlatformUserIdFromUniqueNetId(const FUniqueNetId& UniqueNetId) const override;
 	virtual FString GetAuthType() const override;
 
diff --git a/Engine/Plugins/Online/OnlineSubsystemGoogle/Source/Private/OnlineIdentityGoogleCommon.cpp b/Engine/Plugins/Online/OnlineSubsystemGoogle/Source/Private/OnlineIdentityGoogleCommon.cpp
--- a/Engine/Plugins/Online/OnlineSubsystemGoogle/Source/Private/OnlineIdentityGoogleCommon.cpp
+++ b/Engine/Plugins/Online/OnlineSubsystemGoogle/Source/Private/OnlineIdentityGoogleCommon.cpp
@@ -493,7 +493,7 @@ FString FOnlineIdentityGoogleCommon::GetAuthToken(int32 LocalUserNum) const
 	return FString();
 }
 
-void FOnlineIdentityGoogleCommon::GetUserPrivilege(const FUniqueNetId& UserId, EUserPrivileges::Type Privilege, const FOnGetUserPrivilegeCompleteDelegate& Delegate)
+void FOnlineIdentityGoogleCommon::GetUserPrivilege(const FUniqueNetId& UserId, EUserPrivileges::Type Privilege, const FOnGetUserPrivilegeCompleteDelegate& Delegate, EShowPrivilegeResolveUI ShowResolveUI)
 {
 	Delegate.ExecuteIfBound(UserId, Privilege, (uint32)EPrivilegeResults::NoFailures);
 }	
diff --git a/Engine/Plugins/Online/OnlineSubsystemGoogle/Source/Private/OnlineIdentityGoogleCommon.h b/Engine/Plugins/Online/OnlineSubsystemGoogle/Source/Private/OnlineIdentityGoogleCommon.h
--- a/Engine/Plugins/Online/OnlineSubsystemGoogle/Source/Private/OnlineIdentityGoogleCommon.h
+++ b/Engine/Plugins/Online/OnlineSubsystemGoogle/Source/Private/OnlineIdentityGoogleCommon.h
@@ -68,7 +68,7 @@ public:
 	virtual FString GetPlayerNickname(int32 LocalUserNum) const override;
 	virtual FString GetPlayerNickname(const FUniqueNetId& UserId) const override;
 	virtual FString GetAuthToken(int32 LocalUserNum) const override;
-	virtual void GetUserPrivilege(const FUniqueNetId& UserId, EUserPrivileges::Type Privilege, const FOnGetUserPrivilegeCompleteDelegate& Delegate) override;
+	virtual void GetUserPrivilege(const FUniqueNetId& UserId, EUserPrivileges::Type Privilege, const FOnGetUserPrivilegeCompleteDelegate& Delegate, EShowPrivilegeResolveUI ShowResolveUI=EShowPrivilegeResolveUI::Default) override;
 	virtual FPlatformUserId GetPlatformUserIdFromUniqueNetId(const FUniqueNetId& UniqueNetId) const override;
 	virtual FString GetAuthType() const override;
 	virtual void RevokeAuthToken(const FUniqueNetId& UserId, const FOnRevokeAuthTokenCompleteDelegate& Delegate) override;
diff --git a/Engine/Plugins/Online/OnlineSubsystemNull/Source/Private/OnlineIdentityNull.cpp b/Engine/Plugins/Online/OnlineSubsystemNull/Source/Private/OnlineIdentityNull.cpp
--- a/Engine/Plugins/Online/OnlineSubsystemNull/Source/Private/OnlineIdentityNull.cpp
+++ b/Engine/Plugins/Online/OnlineSubsystemNull/Source/Private/OnlineIdentityNull.cpp
@@ -351,7 +351,7 @@ FOnlineIdentityNull::~FOnlineIdentityNull()
 {
 }
 
-void FOnlineIdentityNull::GetUserPrivilege(const FUniqueNetId& UserId, EUserPrivileges::Type Privilege, const FOnGetUserPrivilegeCompleteDelegate& Delegate)
+void FOnlineIdentityNull::GetUserPrivilege(const FUniqueNetId& UserId, EUserPrivileges::Type Privilege, const FOnGetUserPrivilegeCompleteDelegate& Delegate, EShowPrivilegeResolveUI ShowResolveUI)
 {
 	if (FOnlineSubsystemNull::bForceOfflineMode && Privilege == EUserPrivileges::CanPlayOnline)
 	{
diff --git a/Engine/Plugins/Online/OnlineSubsystemNull/Source/Private/OnlineIdentityNull.h b/Engine/Plugins/Online/OnlineSubsystemNull/Source/Private/OnlineIdentityNull.h
--- a/Engine/Plugins/Online/OnlineSubsystemNull/Source/Private/OnlineIdentityNull.h
+++ b/Engine/Plugins/Online/OnlineSubsystemNull/Source/Private/OnlineIdentityNull.h
@@ -71,7 +71,7 @@ public:
 	virtual FString GetPlayerNickname(const FUniqueNetId& UserId) const override;
 	virtual FString GetAuthToken(int32 LocalUserNum) const override;
 	virtual void RevokeAuthToken(const FUniqueNetId& UserId, const FOnRevokeAuthTokenCompleteDelegate& Delegate) override;
-	virtual void GetUserPrivilege(const FUniqueNetId& UserId, EUserPrivileges::Type Privilege, const FOnGetUserPrivilegeCompleteDelegate& Delegate) override;
+	virtual void GetUserPrivilege(const FUniqueNetId& UserId, EUserPrivileges::Type Privilege, const FOnGetUserPrivilegeCompleteDelegate& Delegate, EShowPrivilegeResolveUI ShowResolveUI=EShowPrivilegeResolveUI::Default) override;
 	virtual FPlatformUserId GetPlatformUserIdFromUniqueNetId(const FUniqueNetId& UniqueNetId) const override;
 	virtual FString GetAuthType() const override;
 
diff --git a/Engine/Plugins/Online/OnlineSubsystemOculus/Source/Private/OnlineIdentityOculus.cpp b/Engine/Plugins/Online/OnlineSubsystemOculus/Source/Private/OnlineIdentityOculus.cpp
--- a/Engine/Plugins/Online/OnlineSubsystemOculus/Source/Private/OnlineIdentityOculus.cpp
+++ b/Engine/Plugins/Online/OnlineSubsystemOculus/Source/Private/OnlineIdentityOculus.cpp
@@ -264,7 +264,7 @@ FOnlineIdentityOculus::FOnlineIdentityOculus(FOnlineSubsystemOculus& InSubsystem
 }
 PRAGMA_ENABLE_DEPRECATION_WARNINGS
 
-void FOnlineIdentityOculus::GetUserPrivilege(const FUniqueNetId& UserId, EUserPrivileges::Type Privilege, const FOnGetUserPrivilegeCompleteDelegate& Delegate)
+void FOnlineIdentityOculus::GetUserPrivilege(const FUniqueNetId& UserId, EUserPrivileges::Type Privilege, const FOnGetUserPrivilegeCompleteDelegate& Delegate, EShowPrivilegeResolveUI ShowResolveUI)
 {
 PRAGMA_DISABLE_DEPRECATION_WARNINGS
 	// Check for entitlement
diff --git a/Engine/Plugins/Online/OnlineSubsystemOculus/Source/Private/OnlineIdentityOculus.h b/Engine/Plugins/Online/OnlineSubsystemOculus/Source/Private/OnlineIdentityOculus.h
--- a/Engine/Plugins/Online/OnlineSubsystemOculus/Source/Private/OnlineIdentityOculus.h
+++ b/Engine/Plugins/Online/OnlineSubsystemOculus/Source/Private/OnlineIdentityOculus.h
@@ -76,7 +76,7 @@ public:
 	virtual FString GetPlayerNickname(const FUniqueNetId& UserId) const override;
 	virtual FString GetAuthToken(int32 LocalUserNum) const override;
 	virtual void RevokeAuthToken(const FUniqueNetId& UserId, const FOnRevokeAuthTokenCompleteDelegate& Delegate) override;
-	virtual void GetUserPrivilege(const FUniqueNetId& UserId, EUserPrivileges::Type Privilege, const FOnGetUserPrivilegeCompleteDelegate& Delegate) override;
+	virtual void GetUserPrivilege(const FUniqueNetId& UserId, EUserPrivileges::Type Privilege, const FOnGetUserPrivilegeCompleteDelegate& Delegate, EShowPrivilegeResolveUI ShowResolveUI=EShowPrivilegeResolveUI::Default) override;
 	virtual FPlatformUserId GetPlatformUserIdFromUniqueNetId(const FUniqueNetId& UniqueNetId) const override;
 	virtual FString GetAuthType() const override;
 
diff --git a/Engine/Plugins/Online/OnlineSubsystemSteam/Source/Private/OnlineIdentityInterfaceSteam.cpp b/Engine/Plugins/Online/OnlineSubsystemSteam/Source/Private/OnlineIdentityInterfaceSteam.cpp
--- a/Engine/Plugins/Online/OnlineSubsystemSteam/Source/Private/OnlineIdentityInterfaceSteam.cpp
+++ b/Engine/Plugins/Online/OnlineSubsystemSteam/Source/Private/OnlineIdentityInterfaceSteam.cpp
@@ -219,7 +219,7 @@ void FOnlineIdentitySteam::RevokeAuthToken(const FUniqueNetId& UserId, const FOn
 	});
 }
 
-void FOnlineIdentitySteam::GetUserPrivilege(const FUniqueNetId& UserId, EUserPrivileges::Type Privilege, const FOnGetUserPrivilegeCompleteDelegate& Delegate)
+void FOnlineIdentitySteam::GetUserPrivilege(const FUniqueNetId& UserId, EUserPrivileges::Type Privilege, const FOnGetUserPrivilegeCompleteDelegate& Delegate, EShowPrivilegeResolveUI ShowResolveUI)
 {
 	Delegate.ExecuteIfBound(UserId, Privilege, (uint32)EPrivilegeResults::NoFailures);
 }
diff --git a/Engine/Plugins/Online/OnlineSubsystemSteam/Source/Private/OnlineIdentityInterfaceSteam.h b/Engine/Plugins/Online/OnlineSubsystemSteam/Source/Private/OnlineIdentityInterfaceSteam.h
--- a/Engine/Plugins/Online/OnlineSubsystemSteam/Source/Private/OnlineIdentityInterfaceSteam.h
+++ b/Engine/Plugins/Online/OnlineSubsystemSteam/Source/Private/OnlineIdentityInterfaceSteam.h
@@ -43,7 +43,7 @@ public:
 	virtual FString GetPlayerNickname(const FUniqueNetId& UserId) const override;
 	virtual FString GetAuthToken(int32 LocalUserNum) const override;
 	virtual void RevokeAuthToken(const FUniqueNetId& UserId, const FOnRevokeAuthTokenCompleteDelegate& Delegate) override;
-	virtual void GetUserPrivilege(const FUniqueNetId& UserId, EUserPrivileges::Type Privilege, const FOnGetUserPrivilegeCompleteDelegate& Delegate) override;
+	virtual void GetUserPrivilege(const FUniqueNetId& UserId, EUserPrivileges::Type Privilege, const FOnGetUserPrivilegeCompleteDelegate& Delegate, EShowPrivilegeResolveUI ShowResolveUI=EShowPrivilegeResolveUI::Default) override;
 	virtual FPlatformUserId GetPlatformUserIdFromUniqueNetId(const FUniqueNetId& UniqueNetId) const override;
 	virtual FString GetAuthType() const override;
 	virtual void GetLinkedAccountAuthToken(int32 LocalUserNum, const FString& TokenType, const FOnGetLinkedAccountAuthTokenCompleteDelegate& Delegate) const override;
diff --git a/Engine/Plugins/Online/OnlineSubsystemTencent/Source/Private/OnlineIdentityTencentCommon.cpp b/Engine/Plugins/Online/OnlineSubsystemTencent/Source/Private/OnlineIdentityTencentCommon.cpp
--- a/Engine/Plugins/Online/OnlineSubsystemTencent/Source/Private/OnlineIdentityTencentCommon.cpp
+++ b/Engine/Plugins/Online/OnlineSubsystemTencent/Source/Private/OnlineIdentityTencentCommon.cpp
@@ -35,7 +35,7 @@ ELoginStatus::Type FOnlineIdentityTencent::GetLoginStatus(const FUniqueNetId& Us
 FString FOnlineIdentityTencent::GetPlayerNickname(int32 LocalUserNum) const{ return FString(); }
 FString FOnlineIdentityTencent::GetPlayerNickname(const FUniqueNetId& UserId) const{ return FString(); }
 FString FOnlineIdentityTencent::GetAuthToken(int32 LocalUserNum) const{ return FString(); }
-void FOnlineIdentityTencent::GetUserPrivilege(const FUniqueNetId& UserId, EUserPrivileges::Type Privilege, const FOnGetUserPrivilegeCompleteDelegate& Delegate){}
+void FOnlineIdentityTencent::GetUserPrivilege(const FUniqueNetId& UserId, EUserPrivileges::Type Privilege, const FOnGetUserPrivilegeCompleteDelegate& Delegate, EShowPrivilegeResolveUI ShowResolveUI){}
 FPlatformUserId FOnlineIdentityTencent::GetPlatformUserIdFromUniqueNetId(const FUniqueNetId& UniqueNetId) const { return FPlatformUserId();  }
 FString FOnlineIdentityTencent::GetAuthType() const{ return FString(); }
 void FOnlineIdentityTencent::RevokeAuthToken(const FUniqueNetId& UserId, const FOnRevokeAuthTokenCompleteDelegate& Delegate){ }
diff --git a/Engine/Plugins/Online/OnlineSubsystemTencent/Source/Private/OnlineIdentityTencentRail.cpp b/Engine/Plugins/Online/OnlineSubsystemTencent/Source/Private/OnlineIdentityTencentRail.cpp
--- a/Engine/Plugins/Online/OnlineSubsystemTencent/Source/Private/OnlineIdentityTencentRail.cpp
+++ b/Engine/Plugins/Online/OnlineSubsystemTencent/Source/Private/OnlineIdentityTencentRail.cpp
@@ -348,7 +348,7 @@ FPlatformUserId FOnlineIdentityTencent::GetPlatformUserIdFromUniqueNetId(const F
 	return PLATFORMUSERID_NONE;
 }
 
-void FOnlineIdentityTencent::GetUserPrivilege(const FUniqueNetId& UserId, EUserPrivileges::Type Privilege, const FOnGetUserPrivilegeCompleteDelegate& Delegate)
+void FOnlineIdentityTencent::GetUserPrivilege(const FUniqueNetId& UserId, EUserPrivileges::Type Privilege, const FOnGetUserPrivilegeCompleteDelegate& Delegate, EShowPrivilegeResolveUI ShowResolveUI)
 {
 	Delegate.ExecuteIfBound(UserId, Privilege, (uint32)EPrivilegeResults::NoFailures);
 }
diff --git a/Engine/Plugins/Online/OnlineSubsystemTencent/Source/Public/OnlineIdentityTencent.h b/Engine/Plugins/Online/OnlineSubsystemTencent/Source/Public/OnlineIdentityTencent.h
--- a/Engine/Plugins/Online/OnlineSubsystemTencent/Source/Public/OnlineIdentityTencent.h
+++ b/Engine/Plugins/Online/OnlineSubsystemTencent/Source/Public/OnlineIdentityTencent.h
@@ -115,7 +115,7 @@ public:
 	virtual FString GetPlayerNickname(int32 LocalUserNum) const override;
 	virtual FString GetPlayerNickname(const FUniqueNetId& UserId) const override;
 	virtual FString GetAuthToken(int32 LocalUserNum) const override;
-	virtual void GetUserPrivilege(const FUniqueNetId& UserId, EUserPrivileges::Type Privilege, const FOnGetUserPrivilegeCompleteDelegate& Delegate) override;
+	virtual void GetUserPrivilege(const FUniqueNetId& UserId, EUserPrivileges::Type Privilege, const FOnGetUserPrivilegeCompleteDelegate& Delegate, EShowPrivilegeResolveUI ShowResolveUI=EShowPrivilegeResolveUI::Default) override;
 	virtual FPlatformUserId GetPlatformUserIdFromUniqueNetId(const FUniqueNetId& UniqueNetId) const override;
 	virtual FString GetAuthType() const override;
 	virtual void RevokeAuthToken(const FUniqueNetId& UserId, const FOnRevokeAuthTokenCompleteDelegate& Delegate) override;
