Introduced a new property to the FPurchaseOfferEntry struct in OnlinePurchaseInterface to be able to add platform specific properties to the checkout flow.

diff --git a/Engine/Plugins/Online/OnlineSubsystem/Source/Public/Interfaces/OnlinePurchaseInterface.h b/Engine/Plugins/Online/OnlineSubsystem/Source/Public/Interfaces/OnlinePurchaseInterface.h
--- a/Engine/Plugins/Online/OnlineSubsystem/Source/Public/Interfaces/OnlinePurchaseInterface.h
+++ b/Engine/Plugins/Online/OnlineSubsystem/Source/Public/Interfaces/OnlinePurchaseInterface.h
@@ -65,6 +65,9 @@ public:
 
 	/** Optional AccountId to include with the purchase request (if supported) */
 	TOptional<FString> AccountId;
+	
+	/** Key - value pair properties for passing to the checkout. */
+	TMap<FString, FString> AdditionalCheckoutProperties;
 
 };
 
