- Remove Embedded OSS names

diff --git a/Engine/Plugins/Online/OnlineSubsystem/Source/Private/OnlineSubsystem.cpp b/Engine/Plugins/Online/OnlineSubsystem/Source/Private/OnlineSubsystem.cpp
--- a/Engine/Plugins/Online/OnlineSubsystem/Source/Private/OnlineSubsystem.cpp
+++ b/Engine/Plugins/Online/OnlineSubsystem/Source/Private/OnlineSubsystem.cpp
@@ -111,10 +111,6 @@ const FName THUNDERHEAD_SUBSYSTEM(TEXT("THUNDERHEAD"));
 const FName MCP_SUBSYSTEM(TEXT("MCP"));
 #endif
 
-#ifndef MCP_SUBSYSTEM_EMBEDDED
-const FName MCP_SUBSYSTEM_EMBEDDED(TEXT("MCP:EMBEDDED"));
-#endif
-
 #ifndef TENCENT_SUBSYSTEM
 const FName TENCENT_SUBSYSTEM(TEXT("TENCENT"));
 #endif
diff --git a/Engine/Plugins/Online/OnlineSubsystem/Source/Public/OnlineSubsystemNames.h b/Engine/Plugins/Online/OnlineSubsystem/Source/Public/OnlineSubsystemNames.h
--- a/Engine/Plugins/Online/OnlineSubsystem/Source/Public/OnlineSubsystemNames.h
+++ b/Engine/Plugins/Online/OnlineSubsystem/Source/Public/OnlineSubsystemNames.h
@@ -69,10 +69,6 @@ extern ONLINESUBSYSTEM_API const FName THUNDERHEAD_SUBSYSTEM;
 extern ONLINESUBSYSTEM_API const FName MCP_SUBSYSTEM;
 #endif
 
-#ifndef MCP_SUBSYSTEM_EMBEDDED
-extern ONLINESUBSYSTEM_API const FName MCP_SUBSYSTEM_EMBEDDED;
-#endif
-
 #ifndef TENCENT_SUBSYSTEM
 extern ONLINESUBSYSTEM_API const FName TENCENT_SUBSYSTEM;
 #endif
