Added new test command line option GETFRIENDDISPLAYNAME to FOnlineSubsystemImpl::HandleFriendExecCommands

diff --git a/Engine/Plugins/Online/OnlineSubsystem/Source/Private/OnlineSubsystemImpl.cpp b/Engine/Plugins/Online/OnlineSubsystem/Source/Private/OnlineSubsystemImpl.cpp
--- a/Engine/Plugins/Online/OnlineSubsystem/Source/Private/OnlineSubsystemImpl.cpp
+++ b/Engine/Plugins/Online/OnlineSubsystem/Source/Private/OnlineSubsystemImpl.cpp
@@ -602,6 +602,96 @@ bool FOnlineSubsystemImpl::HandleFriendExecCommands(UWorld* InWorld, const TCHAR
 		}
 		bWasHandled = true;
 	}
+	if (FParse::Command(&Cmd, TEXT("READFRIENDSLIST"))) /* ONLINE FRIEND READFRIENDSLIST LocalUserNum=0 ListName=default */
+	{
+		IOnlineFriendsPtr FriendsInt = GetFriendsInterface();
+		if (FriendsInt.IsValid())
+		{
+			int LocalUserNum = 0;
+			FParse::Value(Cmd, TEXT("LocalUserNum="), LocalUserNum);
+
+			FString ListName;
+			FParse::Value(Cmd, TEXT("ListName="), ListName);
+
+			FriendsInt->ReadFriendsList(LocalUserNum, ListName, FOnReadFriendsListComplete());
+		}
+		bWasHandled = true;
+	}
+	else if (FParse::Command(&Cmd, TEXT("DUMPFRIENDSLIST"))) /* ONLINE FRIEND DUMPFRIENDSLIST LocalUserNum=0 ListName=default */
+	{
+		IOnlineFriendsPtr FriendsInt = GetFriendsInterface();
+		if (FriendsInt.IsValid())
+		{
+			int LocalUserNum = 0;
+			FParse::Value(Cmd, TEXT("LocalUserNum="), LocalUserNum);
+
+			FString ListName;
+			FParse::Value(Cmd, TEXT("ListName="), ListName);
+
+			TArray< TSharedRef<FOnlineFriend> > Friends;
+			if (FriendsInt->GetFriendsList(LocalUserNum, ListName, Friends))
+			{
+				UE_LOG_ONLINE_FRIEND(Log, TEXT("GetFriendsList returned %d friends"), Friends.Num());
+
+				// Log each friend's data out
+				for (int32 Index = 0; Index < Friends.Num(); Index++)
+				{
+					const FOnlineFriend& Friend = *Friends[Index];
+					const FOnlineUserPresence& Presence = Friend.GetPresence();
+					UE_LOG_ONLINE_FRIEND(Log, TEXT("\t%s has unique id (%s)"), *Friend.GetDisplayName(), *Friend.GetUserId()->ToDebugString());
+					UE_LOG_ONLINE_FRIEND(Log, TEXT("\t\t Invite status (%s)"), EInviteStatus::ToString(Friend.GetInviteStatus()));
+					UE_LOG_ONLINE_FRIEND(Log, TEXT("\t\t Presence: %s"), *Presence.Status.StatusStr);
+					UE_LOG_ONLINE_FRIEND(Log, TEXT("\t\t State: %s"), EOnlinePresenceState::ToString(Presence.Status.State));
+					UE_LOG_ONLINE_FRIEND(Log, TEXT("\t\t bIsOnline (%s)"), *LexToString(Presence.bIsOnline));
+					UE_LOG_ONLINE_FRIEND(Log, TEXT("\t\t bIsPlaying (%s)"), *LexToString(Presence.bIsPlaying));
+					UE_LOG_ONLINE_FRIEND(Log, TEXT("\t\t bIsPlayingThisGame (%s)"), *LexToString(Presence.bIsPlayingThisGame));
+					UE_LOG_ONLINE_FRIEND(Log, TEXT("\t\t bIsJoinable (%s)"), *LexToString(Presence.bIsJoinable));
+					UE_LOG_ONLINE_FRIEND(Log, TEXT("\t\t bHasVoiceSupport (%s)"), *LexToString(Presence.bHasVoiceSupport));
+				}
+			}
+			else
+			{
+				UE_LOG_ONLINE(Warning, TEXT("[%hs] Friend list with name %s was not found"), __FUNCTION__, *ListName);
+			}
+		}
+		bWasHandled = true;
+	}
+	else if (FParse::Command(&Cmd, TEXT("GETDISPLAYNAME"))) /* ONLINE FRIEND GETDISPLAYNAME LocalUserNum=0 ListName=default FriendIndex=0 Platform=Epic */
+	{
+		IOnlineFriendsPtr FriendsInt = GetFriendsInterface();
+		if (FriendsInt.IsValid())
+		{
+			int LocalUserNum = 0;
+			FParse::Value(Cmd, TEXT("LocalUserNum="), LocalUserNum);
+
+			FString ListName;
+			FParse::Value(Cmd, TEXT("ListName="), ListName);
+
+			int FriendIndex = 0;
+			FParse::Value(Cmd, TEXT("FriendIndex="), FriendIndex);
+
+			TArray<TSharedRef<FOnlineFriend>> FriendList;
+			FriendsInt->GetFriendsList(LocalUserNum, ListName, FriendList);
+
+			if (FriendList.IsValidIndex(FriendIndex))
+			{
+				const TSharedRef<FOnlineFriend>& Friend = FriendList[FriendIndex];
+
+				FString Platform;
+				FParse::Value(Cmd, TEXT("Platform="), Platform);
+
+				const FString DisplayName = Friend->GetDisplayName(Platform);
+
+				UE_LOG_ONLINE(Log, TEXT("[%hs] GetDisplayName(%s) call for user %s returned with: %s"), __FUNCTION__, *Platform, *Friend->GetUserId()->ToDebugString(), *DisplayName);
+			}
+			else
+			{
+				UE_LOG_ONLINE(Warning, TEXT("[%hs] Index %d not valid for friend list %s"), __FUNCTION__, FriendIndex, *ListName);
+			}
+		}
+
+		bWasHandled = true;
+	}
 
 	return bWasHandled;
 }
