- Fix TryLoadSubsystemAndSetDefault logging a warning when an OSS is disabled by command line.

diff --git a/Engine/Plugins/Online/OnlineSubsystem/Source/Private/OnlineSubsystemModule.cpp b/Engine/Plugins/Online/OnlineSubsystem/Source/Private/OnlineSubsystemModule.cpp
--- a/Engine/Plugins/Online/OnlineSubsystem/Source/Private/OnlineSubsystemModule.cpp
+++ b/Engine/Plugins/Online/OnlineSubsystem/Source/Private/OnlineSubsystemModule.cpp
@@ -155,31 +155,38 @@ bool FOnlineSubsystemModule::TryLoadSubsystemAndSetDefault(FName SubsystemName)
 {
 	// A module loaded with its factory method set for creation and a default instance of the online subsystem is required
 	bool bLoaded = false;
-	FString SubsystemNameString = SubsystemName.ToString();
-	if (LoadSubsystemModule(SubsystemNameString, ModuleRedirects))
+	const FString SubsystemNameString = SubsystemName.ToString();
+	if (IOnlineSubsystem::IsEnabled(SubsystemName))
 	{
-		if (OnlineFactories.Contains(SubsystemName))
+		if (LoadSubsystemModule(SubsystemNameString, ModuleRedirects))
 		{
-			IOnlineSubsystem* OnlineSubsystem = GetOnlineSubsystem(SubsystemName);
-			if (OnlineSubsystem != nullptr)
+			if (OnlineFactories.Contains(SubsystemName))
 			{
-				UE_LOG_ONLINE(Log, TEXT("TryLoadSubsystemAndSetDefault: Loaded subsystem for type [%s]"), *SubsystemNameString);
-				DefaultPlatformService = SubsystemName;
-				bLoaded = true;
+				IOnlineSubsystem* OnlineSubsystem = GetOnlineSubsystem(SubsystemName);
+				if (OnlineSubsystem != nullptr)
+				{
+					UE_LOG_ONLINE(Log, TEXT("TryLoadSubsystemAndSetDefault: Loaded subsystem for type [%s]"), *SubsystemNameString);
+					DefaultPlatformService = SubsystemName;
+					bLoaded = true;
+				}
+				else
+				{
+					//UE_LOG_ONLINE(Warning, TEXT("TryLoadSubsystemAndSetDefault: GetOnlineSubsystem([%s]) failed"), *SubsystemNameString);
+				}
 			}
 			else
 			{
-				//UE_LOG_ONLINE(Warning, TEXT("TryLoadSubsystemAndSetDefault: GetOnlineSubsystem([%s]) failed"), *SubsystemNameString);
+				UE_LOG_ONLINE(Warning, TEXT("TryLoadSubsystemAndSetDefault: OnlineFactories does not contain [%s]"), *SubsystemNameString);
 			}
 		}
 		else
 		{
-			UE_LOG_ONLINE(Warning, TEXT("TryLoadSubsystemAndSetDefault: OnlineFactories does not contain [%s]"), *SubsystemNameString);
+			UE_LOG_ONLINE(Warning, TEXT("TryLoadSubsystemAndSetDefault: LoadSubsystemModule([%s]) failed"), *SubsystemNameString);
 		}
 	}
 	else
 	{
-		UE_LOG_ONLINE(Warning, TEXT("TryLoadSubsystemAndSetDefault: LoadSubsystemModule([%s]) failed"), *SubsystemNameString);
+		UE_LOG_ONLINE(Verbose, TEXT("TryLoadSubsystemAndSetDefault: [%s] disabled"), *SubsystemNameString);
 	}
 
 	return bLoaded;
