- Bump disabled logging from Verbose to Log

diff --git a/Engine/Plugins/Online/OnlineSubsystem/Source/Private/OnlineSubsystemModule.cpp b/Engine/Plugins/Online/OnlineSubsystem/Source/Private/OnlineSubsystemModule.cpp
--- a/Engine/Plugins/Online/OnlineSubsystem/Source/Private/OnlineSubsystemModule.cpp
+++ b/Engine/Plugins/Online/OnlineSubsystem/Source/Private/OnlineSubsystemModule.cpp
@@ -186,7 +186,7 @@ bool FOnlineSubsystemModule::TryLoadSubsystemAndSetDefault(FName SubsystemName)
 	}
 	else
 	{
-		UE_LOG_ONLINE(Verbose, TEXT("TryLoadSubsystemAndSetDefault: [%s] disabled"), *SubsystemNameString);
+		UE_LOG_ONLINE(Log, TEXT("TryLoadSubsystemAndSetDefault: [%s] disabled"), *SubsystemNameString);
 	}
 
 	return bLoaded;
