- Add [Room|Private]MessageSent events

diff --git a/Engine/Plugins/Online/OnlineFramework/Source/Party/Public/Chat/SocialChatManager.cpp b/Engine/Plugins/Online/OnlineFramework/Source/Party/Public/Chat/SocialChatManager.cpp
--- a/Engine/Plugins/Online/OnlineFramework/Source/Party/Public/Chat/SocialChatManager.cpp
+++ b/Engine/Plugins/Online/OnlineFramework/Source/Party/Public/Chat/SocialChatManager.cpp
@@ -254,7 +254,9 @@ void USocialChatManager::InitializeChatManager()
 		ChatInterface->AddOnChatRoomMemberJoinDelegate_Handle(FOnChatRoomMemberJoinDelegate::CreateUObject(this, &USocialChatManager::HandleChatRoomMemberJoin));
 		ChatInterface->AddOnChatRoomMemberExitDelegate_Handle(FOnChatRoomMemberExitDelegate::CreateUObject(this, &USocialChatManager::HandleChatRoomMemberExit));
 		ChatInterface->AddOnChatRoomMemberUpdateDelegate_Handle(FOnChatRoomMemberUpdateDelegate::CreateUObject(this, &USocialChatManager::HandleChatRoomMemberUpdate));
+		ChatInterface->AddOnChatRoomMessageSentDelegate_Handle(FOnChatRoomMessageSentDelegate::CreateUObject(this, &USocialChatManager::HandleChatRoomMessageSent));
 		ChatInterface->AddOnChatRoomMessageReceivedDelegate_Handle(FOnChatRoomMessageReceivedDelegate::CreateUObject(this, &USocialChatManager::HandleChatRoomMessageReceived));
+		ChatInterface->AddOnChatPrivateMessageSentDelegate_Handle(FOnChatPrivateMessageSentDelegate::CreateUObject(this, &USocialChatManager::HandleChatPrivateMessageSent));
 		ChatInterface->AddOnChatPrivateMessageReceivedDelegate_Handle(FOnChatPrivateMessageReceivedDelegate::CreateUObject(this, &USocialChatManager::HandleChatPrivateMessageReceived));
 	}
 	else
@@ -461,6 +463,10 @@ void USocialChatManager::HandleChatRoomMemberUpdate(const FUniqueNetId& LocalUse
 	}
 }
 
+void USocialChatManager::HandleChatRoomMessageSent(const FUniqueNetId& LocalUserId, const FChatRoomId& RoomId, bool bWasSuccessful, const FString& Error)
+{
+}
+
 void USocialChatManager::HandleChatRoomMessageReceived(const FUniqueNetId& LocalUserId, const FChatRoomId& RoomId, const TSharedRef<FChatMessage>& ChatMessage)
 {
 	if (IsUniqueIdOfOwner(LocalUserId))
@@ -473,6 +479,10 @@ void USocialChatManager::HandleChatRoomMessageReceived(const FUniqueNetId& Local
 	}
 }
 
+void USocialChatManager::HandleChatPrivateMessageSent(const FUniqueNetId& LocalUserId, const FUniqueNetId& RecipientId, bool bWasSuccessful, const FString& Error)
+{
+}
+
 void USocialChatManager::HandleChatPrivateMessageReceived(const FUniqueNetId& LocalUserId, const TSharedRef<FChatMessage>& ChatMessage)
 {
 	if (IsUniqueIdOfOwner(LocalUserId))
diff --git a/Engine/Plugins/Online/OnlineFramework/Source/Party/Public/Chat/SocialChatManager.h b/Engine/Plugins/Online/OnlineFramework/Source/Party/Public/Chat/SocialChatManager.h
--- a/Engine/Plugins/Online/OnlineFramework/Source/Party/Public/Chat/SocialChatManager.h
+++ b/Engine/Plugins/Online/OnlineFramework/Source/Party/Public/Chat/SocialChatManager.h
@@ -140,7 +140,9 @@ protected:
 	virtual void InitializeChatManager();
 	virtual ESocialChannelType TryChannelTypeLookupByRoomId(const FChatRoomId& RoomID);
 
+	virtual void HandleChatRoomMessageSent(const FUniqueNetId& LocalUserId, const FChatRoomId& RoomId, bool bWasSuccessful, const FString& Error);
 	virtual void HandleChatRoomMessageReceived(const FUniqueNetId& LocalUserId, const FChatRoomId& RoomId, const TSharedRef<FChatMessage>& ChatMessage);
+	virtual void HandleChatPrivateMessageSent(const FUniqueNetId& LocalUserId, const FUniqueNetId& RecipientId, bool bWasSuccessful, const FString& Error);
 	virtual void HandleChatPrivateMessageReceived(const FUniqueNetId& LocalUserId, const TSharedRef<FChatMessage>& ChatMessage);
 
 	virtual void OnChannelCreatedInternal(USocialChatChannel& CreatedChannel);
diff --git a/Engine/Plugins/Online/OnlineSubsystem/Source/Public/Interfaces/OnlineChatInterface.h b/Engine/Plugins/Online/OnlineSubsystem/Source/Public/Interfaces/OnlineChatInterface.h
--- a/Engine/Plugins/Online/OnlineSubsystem/Source/Public/Interfaces/OnlineChatInterface.h
+++ b/Engine/Plugins/Online/OnlineSubsystem/Source/Public/Interfaces/OnlineChatInterface.h
@@ -199,6 +199,17 @@ typedef FOnChatRoomMemberExit::FDelegate FOnChatRoomMemberExitDelegate;
 DECLARE_MULTICAST_DELEGATE_ThreeParams(FOnChatRoomMemberUpdate, const FUniqueNetId& /*UserId*/, const FChatRoomId& /*RoomId*/, const FUniqueNetId& /*MemberId*/);
 typedef FOnChatRoomMemberUpdate::FDelegate FOnChatRoomMemberUpdateDelegate;
 
+/**
+* Delegate used when a chat message is sent to a chat room
+*
+* @param UserId the user that made the request
+* @param RoomId room that was requested
+* @param bWasSuccessful true if the async action completed without error, false if there was an error
+* @param Error string representing the error condition
+*/
+DECLARE_MULTICAST_DELEGATE_FourParams(FOnChatRoomMessageSent, const FUniqueNetId& /*UserId*/, const FChatRoomId& /*RoomId*/, bool /*bWasSuccessful*/, const FString& /*Error*/);
+typedef FOnChatRoomMessageSent::FDelegate FOnChatRoomMessageSentDelegate;
+
 /**
  * Delegate used when a chat message is received from a chat room
  *
@@ -209,6 +220,17 @@ typedef FOnChatRoomMemberUpdate::FDelegate FOnChatRoomMemberUpdateDelegate;
 DECLARE_MULTICAST_DELEGATE_ThreeParams(FOnChatRoomMessageReceived, const FUniqueNetId& /*UserId*/, const FChatRoomId& /*RoomId*/, const TSharedRef<FChatMessage>& /*ChatMessage*/);
 typedef FOnChatRoomMessageReceived::FDelegate FOnChatRoomMessageReceivedDelegate;
 
+/**
+* Delegate used when a private chat message is sent to another user
+*
+* @param UserId the user that made the request
+* @param RecipientId the user the message was sent to
+* @param bWasSuccessful true if the async action completed without error, false if there was an error
+* @param Error string representing the error condition
+*/
+DECLARE_MULTICAST_DELEGATE_FourParams(FOnChatPrivateMessageSent, const FUniqueNetId& /*UserId*/, const FUniqueNetId& /*RecipientId*/, bool /*bWasSuccessful*/, const FString& /*Error*/);
+typedef FOnChatPrivateMessageSent::FDelegate FOnChatPrivateMessageSentDelegate;
+
 /**
  * Delegate used when a private chat message is received from another user
  *
@@ -393,7 +415,9 @@ public:
 	DEFINE_ONLINE_DELEGATE_THREE_PARAM(OnChatRoomMemberJoin, const FUniqueNetId&, const FChatRoomId&, const FUniqueNetId&);
 	DEFINE_ONLINE_DELEGATE_THREE_PARAM(OnChatRoomMemberExit, const FUniqueNetId&, const FChatRoomId&, const FUniqueNetId&);
 	DEFINE_ONLINE_DELEGATE_THREE_PARAM(OnChatRoomMemberUpdate, const FUniqueNetId&, const FChatRoomId&, const FUniqueNetId&);
+	DEFINE_ONLINE_DELEGATE_FOUR_PARAM(OnChatRoomMessageSent, const FUniqueNetId&, const FChatRoomId&, bool, const FString&);
 	DEFINE_ONLINE_DELEGATE_THREE_PARAM(OnChatRoomMessageReceived, const FUniqueNetId&, const FChatRoomId&, const TSharedRef<FChatMessage>&);
+	DEFINE_ONLINE_DELEGATE_FOUR_PARAM(OnChatPrivateMessageSent, const FUniqueNetId&, const FUniqueNetId&, bool, const FString&);
 	DEFINE_ONLINE_DELEGATE_TWO_PARAM(OnChatPrivateMessageReceived, const FUniqueNetId&, const TSharedRef<FChatMessage>&);
 };
 
