Update status for legibility

#rnx
#tests PIE
- Ran PIE and verified status legibility was improved

diff --git a/Engine/Plugins/Online/OnlineSubsystem/Source/Public/Interfaces/OnlinePresenceInterface.h b/Engine/Plugins/Online/OnlineSubsystem/Source/Public/Interfaces/OnlinePresenceInterface.h
--- a/Engine/Plugins/Online/OnlineSubsystem/Source/Public/Interfaces/OnlinePresenceInterface.h
+++ b/Engine/Plugins/Online/OnlineSubsystem/Source/Public/Interfaces/OnlinePresenceInterface.h
@@ -121,7 +121,6 @@ namespace EOnlinePresenceState
 	static FText OnlineText =  NSLOCTEXT("OnlinePresence", "Online", "Online");
 	static FText OfflineText =  NSLOCTEXT("OnlinePresence", "Offline", "Offline");
 	static FText AwayText =  NSLOCTEXT("OnlinePresence", "Away", "Away");
-	static FText ExtendedAwayText =  NSLOCTEXT("OnlinePresence", "ExtendedAway", "Extended Away");
 	static FText DoNotDisturbText =  NSLOCTEXT("OnlinePresence", "DoNotDisturb", "Do Not Disturb");
 	static FText ChatText =  NSLOCTEXT("OnlinePresence", "Chat", "Chat");
 	/** 
@@ -135,10 +134,10 @@ namespace EOnlinePresenceState
 			return OnlineText;
 		case Offline:
 			return OfflineText;
+		case ExtendedAway:
+			// falls through to return away text
 		case Away:
 			return AwayText;
-		case ExtendedAway:
-			return ExtendedAwayText;
 		case DoNotDisturb:
 			return DoNotDisturbText;
 		case Chat:
