Fix build errors with merged modular builds

#rnx

diff --git a/Engine/Plugins/Online/OnlineSubsystem/Source/Public/Interfaces/OnlineUserCloudInterface.h b/Engine/Plugins/Online/OnlineSubsystem/Source/Public/Interfaces/OnlineUserCloudInterface.h
--- a/Engine/Plugins/Online/OnlineSubsystem/Source/Public/Interfaces/OnlineUserCloudInterface.h
+++ b/Engine/Plugins/Online/OnlineSubsystem/Source/Public/Interfaces/OnlineUserCloudInterface.h
@@ -129,12 +129,14 @@ class IOnlineUserCloud
 {
 
 protected:
+	PRAGMA_DISABLE_DEPRECATION_WARNINGS
 	IOnlineUserCloud() {};
+	PRAGMA_ENABLE_DEPRECATION_WARNINGS
 
 public:
-PRAGMA_DISABLE_DEPRECATION_WARNINGS
+	PRAGMA_DISABLE_DEPRECATION_WARNINGS
 	virtual ~IOnlineUserCloud() {};
-PRAGMA_ENABLE_DEPRECATION_WARNINGS
+	PRAGMA_ENABLE_DEPRECATION_WARNINGS
 
 	/**
 	 * Copies the file data into the specified buffer for the specified file
diff --git a/Engine/Source/Runtime/GameplayDebugger/Public/GameplayDebuggerLocalController.h b/Engine/Source/Runtime/GameplayDebugger/Public/GameplayDebuggerLocalController.h
--- a/Engine/Source/Runtime/GameplayDebugger/Public/GameplayDebuggerLocalController.h
+++ b/Engine/Source/Runtime/GameplayDebugger/Public/GameplayDebuggerLocalController.h
@@ -43,7 +43,7 @@ class UGameplayDebuggerLocalController : public UObject
 	/** checks if key is bound by any action */
 	bool GAMEPLAYDEBUGGER_API IsKeyBound(const FName KeyName) const;
 
-	static void SetConsoleCommandsEnabled(const bool bEnabled) { bConsoleCommandsEnabled = bEnabled; }
+	static GAMEPLAYDEBUGGER_API void SetConsoleCommandsEnabled(const bool bEnabled) { bConsoleCommandsEnabled = bEnabled; }
 #endif // WITH_GAMEPLAY_DEBUGGER_MENU	
 protected:
 	UPROPERTY()
diff --git a/Engine/Source/Runtime/MediaUtils/Public/MediaObjectPool.h b/Engine/Source/Runtime/MediaUtils/Public/MediaObjectPool.h
--- a/Engine/Source/Runtime/MediaUtils/Public/MediaObjectPool.h
+++ b/Engine/Source/Runtime/MediaUtils/Public/MediaObjectPool.h
@@ -268,7 +268,7 @@ public:
 	 * @param NumReserve Number of objects to reserve.
 	 */
 	TMediaObjectPool(uint32 NumReserve)
-		: Storage(MakeShareable(new TStorage))
+		: Storage(MakeShareable(new TStorage(nullptr)))
 	{
 		Storage->Reserve(NumReserve);
 	}
