- Add Index to friends list dump console command

#rnx

diff --git a/Engine/Plugins/Online/OnlineSubsystem/Source/Private/OnlineSubsystemImpl.cpp b/Engine/Plugins/Online/OnlineSubsystem/Source/Private/OnlineSubsystemImpl.cpp
--- a/Engine/Plugins/Online/OnlineSubsystem/Source/Private/OnlineSubsystemImpl.cpp
+++ b/Engine/Plugins/Online/OnlineSubsystem/Source/Private/OnlineSubsystemImpl.cpp
@@ -638,7 +638,7 @@ bool FOnlineSubsystemImpl::HandleFriendExecCommands(UWorld* InWorld, const TCHAR
 				{
 					const FOnlineFriend& Friend = *Friends[Index];
 					const FOnlineUserPresence& Presence = Friend.GetPresence();
-					UE_LOG_ONLINE_FRIEND(Log, TEXT("\t%s has unique id (%s)"), *Friend.GetDisplayName(), *Friend.GetUserId()->ToDebugString());
+					UE_LOG_ONLINE_FRIEND(Log, TEXT("\tIdx %d DisplayName (%s) NetId (%s)"), Index, *Friend.GetDisplayName(), *Friend.GetUserId()->ToDebugString());
 					UE_LOG_ONLINE_FRIEND(Log, TEXT("\t\t Invite status (%s)"), EInviteStatus::ToString(Friend.GetInviteStatus()));
 					UE_LOG_ONLINE_FRIEND(Log, TEXT("\t\t Presence: %s"), *Presence.Status.StatusStr);
 					UE_LOG_ONLINE_FRIEND(Log, TEXT("\t\t State: %s"), EOnlinePresenceState::ToString(Presence.Status.State));
