Deprecation of Stat/Leaderboard Interface variables and methods using FName, with new FString variants.
Refactor of callsites to use new methods and variables.

diff --git a/Engine/Plugins/Online/IOS/OnlineSubsystemIOS/Source/Private/OnlineLeaderboardsInterfaceIOS.cpp b/Engine/Plugins/Online/IOS/OnlineSubsystemIOS/Source/Private/OnlineLeaderboardsInterfaceIOS.cpp
--- a/Engine/Plugins/Online/IOS/OnlineSubsystemIOS/Source/Private/OnlineLeaderboardsInterfaceIOS.cpp
+++ b/Engine/Plugins/Online/IOS/OnlineSubsystemIOS/Source/Private/OnlineLeaderboardsInterfaceIOS.cpp
@@ -241,7 +241,7 @@ bool FOnlineLeaderboardsIOS::WriteLeaderboards(const FName& SessionName, const F
 		// Access the stat and the value.
 		const FVariantData& Stat = It.Value();
 
-		FString LeaderboardName(It.Key().ToString());
+		FString LeaderboardName(It.Key());
 		NSString* Category = [NSString stringWithFString:LeaderboardName];
 
 		bool bIsValidScore = false;
diff --git a/Engine/Plugins/Online/OnlineServicesOSSAdapter/Source/Private/Online/LeaderboardsOSSAdapter.cpp b/Engine/Plugins/Online/OnlineServicesOSSAdapter/Source/Private/Online/LeaderboardsOSSAdapter.cpp
--- a/Engine/Plugins/Online/OnlineServicesOSSAdapter/Source/Private/Online/LeaderboardsOSSAdapter.cpp
+++ b/Engine/Plugins/Online/OnlineServicesOSSAdapter/Source/Private/Online/LeaderboardsOSSAdapter.cpp
@@ -240,8 +240,8 @@ TOptional<FOnlineError> FLeaderboardsOSSAdapter::PrepareLeaderboardReadObject(co
 	else
 	{
 		ReadObject = MakeShareable(new FOnlineLeaderboardRead());
-		ReadObject->SortedColumn = *BoardName;
-		ReadObject->LeaderboardName = *BoardName;
+		ReadObject->SortedColumn = BoardName;
+		ReadObject->LeaderboardName = BoardName;
 	}
 
 	return OnlineError;
diff --git a/Engine/Plugins/Online/OnlineSubsystem/Source/Private/OnlineStats.cpp b/Engine/Plugins/Online/OnlineSubsystem/Source/Private/OnlineStats.cpp
--- a/Engine/Plugins/Online/OnlineSubsystem/Source/Private/OnlineStats.cpp
+++ b/Engine/Plugins/Online/OnlineSubsystem/Source/Private/OnlineStats.cpp
@@ -7,7 +7,7 @@
  * @param StatName key name to search for
  * @return KeyValuePair if found, NULL otherwise
  */
-FVariantData* FOnlineStats::FindStatByName(const FName& StatName)
+FVariantData* FOnlineStats::FindStatByName(const FString& StatName)
 {
 	return Properties.Find(StatName);
 }
@@ -19,7 +19,7 @@ FVariantData* FOnlineStats::FindStatByName(const FName& StatName)
  * @param StatName the stat to change the value of
  * @param Value the new value to assign to the stat
  */
-void FOnlineStats::SetFloatStat(const FName& StatName, float Value)
+void FOnlineStats::SetFloatStat(const FString& StatName, float Value)
 {
 	FVariantData* Stat = FindStatByName(StatName);
 	if (Stat != NULL)
@@ -44,7 +44,7 @@ void FOnlineStats::SetFloatStat(const FName& StatName, float Value)
  * @param StatName the stat to change the value of
  * @param Value the new value to assign to the stat
  */
-void FOnlineStats::SetIntStat(const FName& StatName, int32 Value)
+void FOnlineStats::SetIntStat(const FString& StatName, int32 Value)
 {
 	FVariantData* Stat = FindStatByName(StatName);
 	if (Stat != NULL && Stat->GetType() == EOnlineKeyValuePairDataType::Int32)
@@ -66,7 +66,7 @@ void FOnlineStats::SetIntStat(const FName& StatName, int32 Value)
  * @param StatName the stat to increment
  * @param IncBy the value to increment by
  */
-void FOnlineStats::IncrementFloatStat(const FName& StatName, float IncBy)
+void FOnlineStats::IncrementFloatStat(const FString& StatName, float IncBy)
 {
 	FVariantData* Stat = FindStatByName(StatName);
 	if (Stat != NULL && Stat->GetType() == EOnlineKeyValuePairDataType::Float)
@@ -88,7 +88,7 @@ void FOnlineStats::IncrementFloatStat(const FName& StatName, float IncBy)
  * @param StatName the stat to increment
  * @param IncBy the value to increment by
  */
-void FOnlineStats::IncrementIntStat(const FName& StatName, int32 IncBy)
+void FOnlineStats::IncrementIntStat(const FString& StatName, int32 IncBy)
 {
 	FVariantData* Stat = FindStatByName(StatName);
 	if (Stat != NULL && Stat->GetType() == EOnlineKeyValuePairDataType::Int32)
@@ -110,7 +110,7 @@ void FOnlineStats::IncrementIntStat(const FName& StatName, int32 IncBy)
  * @param StatName the stat to decrement
  * @param DecBy the value to decrement by
  */
-void FOnlineStats::DecrementFloatStat(const FName& StatName, float DecBy)
+void FOnlineStats::DecrementFloatStat(const FString& StatName, float DecBy)
 {
 	FVariantData* Stat = FindStatByName(StatName);
 	if (Stat != NULL && Stat->GetType() == EOnlineKeyValuePairDataType::Float)
@@ -132,7 +132,7 @@ void FOnlineStats::DecrementFloatStat(const FName& StatName, float DecBy)
  * @param StatName the stat to decrement
  * @param DecBy the value to decrement by
  */
-void FOnlineStats::DecrementIntStat(const FName& StatName, int32 DecBy)
+void FOnlineStats::DecrementIntStat(const FString& StatName, int32 DecBy)
 {
 	FVariantData* Stat = FindStatByName(StatName);
 	if (Stat != NULL && Stat->GetType() == EOnlineKeyValuePairDataType::Int32)
@@ -151,9 +151,9 @@ FString FOnlineStatsRow::ToLogString() const
 {
 	FString LogString = FString::Printf(TEXT("%d : %s"), Rank, *NickName);
 
-	for (const TPair<FName, FVariantData>& Column : Columns)
+	for (const TPair<FString, FVariantData>& Column : Columns)
 	{
-		LogString += FString::Printf(TEXT("\t\t%s : %s"), *Column.Key.ToString(), *Column.Value.ToString());
+		LogString += FString::Printf(TEXT("\t\t%s : %s"), *Column.Key, *Column.Value.ToString());
 	}
 
 	return LogString;
@@ -161,7 +161,7 @@ FString FOnlineStatsRow::ToLogString() const
 
 FString FOnlineLeaderboardRead::ToLogString() const
 {
-	FString LogString = FString::Printf(TEXT("\nLeaderboardName: %s\nSortedColumn: %s\nRows:\n"), *LeaderboardName.ToString(), *SortedColumn.ToString());
+	FString LogString = FString::Printf(TEXT("\nLeaderboardName: %s\nSortedColumn: %s\nRows:\n"), *LeaderboardName, *SortedColumn);
 
 	for (const FOnlineStatsRow& Row : Rows)
 	{
diff --git a/Engine/Plugins/Online/OnlineSubsystem/Source/Public/OnlineStats.h b/Engine/Plugins/Online/OnlineSubsystem/Source/Public/OnlineStats.h
--- a/Engine/Plugins/Online/OnlineSubsystem/Source/Public/OnlineStats.h
+++ b/Engine/Plugins/Online/OnlineSubsystem/Source/Public/OnlineStats.h
@@ -5,10 +5,179 @@
 #include "OnlineSubsystemTypes.h"
 #include "OnlineKeyValuePair.h"
 
+// typedef FOnlineKeyValuePairs<FString, FVariantData> FStatsColumnArray;
+// Temporary class to assist in deprecation of changing the key from FName to FString. After the deprecation period, the class will be deleted and replaced by the line above.
 /** Representation of a single column and its data */ 
-typedef FOnlineKeyValuePairs<FName, FVariantData> FStatsColumnArray;
+class FStatsColumnArray : public FOnlineKeyValuePairs<FString, FVariantData>
+{
+	typedef FOnlineKeyValuePairs<FString, FVariantData> Super;
+
+public:
+	FORCEINLINE FStatsColumnArray() {}
+	FORCEINLINE FStatsColumnArray(FStatsColumnArray&& Other) : Super(MoveTemp(Other)) {}
+	FORCEINLINE FStatsColumnArray(const FStatsColumnArray& Other) : Super(Other) {}
+	FORCEINLINE FStatsColumnArray& operator=(FStatsColumnArray&& Other) { Super::operator=(MoveTemp(Other)); return *this; }
+	FORCEINLINE FStatsColumnArray& operator=(const FStatsColumnArray& Other) { Super::operator=(Other); return *this; }
+
+	UE_DEPRECATED(5.5, "FStatsColumnArray now uses FString for the key type instead of FName")
+		FStatsColumnArray(const FOnlineKeyValuePairs< FName, FVariantData >& DeprecatedValues)
+	{
+		Reserve(DeprecatedValues.Num());
+		for (const TPair< FName, FVariantData >& DeprecatedValue : DeprecatedValues)
+		{
+			Super::Emplace(DeprecatedValue.Key.ToString(), DeprecatedValue.Value);
+		}
+	}
+
+	UE_DEPRECATED(5.5, "FStatsColumnArray now uses FString for the key type instead of FName")
+		FStatsColumnArray(FOnlineKeyValuePairs< FName, FVariantData >&& DeprecatedValues)
+	{
+		Reserve(DeprecatedValues.Num());
+		for (TPair< FName, FVariantData >& DeprecatedValue : DeprecatedValues)
+		{
+			Super::Emplace(DeprecatedValue.Key.ToString(), MoveTemp(DeprecatedValue.Value));
+		}
+	}
+
+	using Super::Add;
+	UE_DEPRECATED(5.5, "FStatsColumnArray now uses FString for the key type instead of FName")
+		FVariantData& Add(const FName& InKey, const FVariantData& InValue)
+	{
+		return Super::Add(InKey.ToString(), InValue);
+	}
+
+	using Super::Emplace;
+	UE_DEPRECATED(5.5, "FStatsColumnArray now uses FString for the key type instead of FName")
+		FVariantData& Emplace(const FName& InKey, const FVariantData& InValue)
+	{
+		return Super::Emplace(InKey.ToString(), InValue);
+	}
+
+	using Super::Find;
+	UE_DEPRECATED(5.5, "FStatsColumnArray now uses FString for the key type instead of FName")
+		FVariantData* Find(const FName& InKey)
+	{
+		return Super::Find(InKey.ToString());
+	}
+};
+
+
+// typedef FOnlineKeyValuePairs<FString, FVariantData> FStatPropertyArray;
+// Temporary class to assist in deprecation of changing the key from FName to FString. After the deprecation period, the class will be deleted and replaced by the line above.
 /** Representation of a single stat value to post to the backend */
-typedef FOnlineKeyValuePairs<FName, FVariantData> FStatPropertyArray;
+class FStatPropertyArray : public FOnlineKeyValuePairs<FString, FVariantData>
+{
+	typedef FOnlineKeyValuePairs<FString, FVariantData> Super;
+
+public:
+	FORCEINLINE FStatPropertyArray() {}
+	FORCEINLINE FStatPropertyArray(FStatPropertyArray&& Other) : Super(MoveTemp(Other)) {}
+	FORCEINLINE FStatPropertyArray(const FStatPropertyArray& Other) : Super(Other) {}
+	FORCEINLINE FStatPropertyArray& operator=(FStatPropertyArray&& Other) { Super::operator=(MoveTemp(Other)); return *this; }
+	FORCEINLINE FStatPropertyArray& operator=(const FStatPropertyArray& Other) { Super::operator=(Other); return *this; }
+
+	UE_DEPRECATED(5.5, "FStatPropertyArray now uses FString for the key type instead of FName")
+		FStatPropertyArray(const FOnlineKeyValuePairs< FName, FVariantData >& DeprecatedValues)
+	{
+		Reserve(DeprecatedValues.Num());
+		for (const TPair< FName, FVariantData >& DeprecatedValue : DeprecatedValues)
+		{
+			Super::Emplace(DeprecatedValue.Key.ToString(), DeprecatedValue.Value);
+		}
+	}
+
+	UE_DEPRECATED(5.5, "FStatPropertyArray now uses FString for the key type instead of FName")
+		FStatPropertyArray(FOnlineKeyValuePairs< FName, FVariantData >&& DeprecatedValues)
+	{
+		Reserve(DeprecatedValues.Num());
+		for (TPair< FName, FVariantData >& DeprecatedValue : DeprecatedValues)
+		{
+			Super::Emplace(DeprecatedValue.Key.ToString(), MoveTemp(DeprecatedValue.Value));
+		}
+	}
+
+	using Super::Add;
+	UE_DEPRECATED(5.5, "FStatPropertyArray now uses FString for the key type instead of FName")
+		FVariantData& Add(const FName& InKey, const FVariantData& InValue)
+	{
+		return Super::Add(InKey.ToString(), InValue);
+	}
+
+	using Super::Emplace;
+	UE_DEPRECATED(5.5, "FStatPropertyArray now uses FString for the key type instead of FName")
+		FVariantData& Emplace(const FName& InKey, const FVariantData& InValue)
+	{
+		return Super::Emplace(InKey.ToString(), InValue);
+	}
+
+	using Super::Find;
+	UE_DEPRECATED(5.5, "FStatPropertyArray now uses FString for the key type instead of FName")
+		FVariantData* Find(const FName& InKey)
+	{
+		return Super::Find(InKey.ToString());
+	}
+};
+
+class FNameArrayDeprecationWrapper : public TArray<FString>
+{
+	typedef TArray<FString> Super;
+
+public:
+	FORCEINLINE FNameArrayDeprecationWrapper() {}
+	FORCEINLINE FNameArrayDeprecationWrapper(FNameArrayDeprecationWrapper&& Other) : Super(MoveTemp(Other)) {}
+	FORCEINLINE FNameArrayDeprecationWrapper(const FNameArrayDeprecationWrapper& Other) : Super(Other) {}
+	FORCEINLINE FNameArrayDeprecationWrapper& operator=(FNameArrayDeprecationWrapper&& Other) { Super::operator=(MoveTemp(Other)); return *this; }
+	FORCEINLINE FNameArrayDeprecationWrapper& operator=(const FNameArrayDeprecationWrapper& Other) { Super::operator=(Other); return *this; }
+
+	using Super::Add;
+	UE_DEPRECATED(5.5, "This variable is now a TArray<FString> instead of a TArray<FName>. Please update your code accordingly")
+		int32 Add(const FName& InElement)
+	{
+		return Super::Add(InElement.ToString());
+	}
+
+	using Super::Emplace;
+	UE_DEPRECATED(5.5, "This variable is now a TArray<FString> instead of a TArray<FName>. Please update your code accordingly")
+		int32 Emplace(const FName& InElement)
+	{
+		return Super::Emplace(InElement.ToString());
+	}
+};
+
+class FNameDeprecationWrapper : public FString
+{
+	typedef FString Super;
+
+public:
+	FORCEINLINE FNameDeprecationWrapper() {}
+	FORCEINLINE FNameDeprecationWrapper(FNameDeprecationWrapper&& Other) : Super(MoveTemp(Other)) {}
+	FORCEINLINE FNameDeprecationWrapper(const FNameDeprecationWrapper& Other) : Super(Other) {}
+	FORCEINLINE FNameDeprecationWrapper& operator=(FNameDeprecationWrapper&& Other) { Super::operator=(MoveTemp(Other)); return *this; }
+	FORCEINLINE FNameDeprecationWrapper& operator=(const FNameDeprecationWrapper& Other) { Super::operator=(Other); return *this; }
+
+	FORCEINLINE FNameDeprecationWrapper(FString&& Other) : Super(MoveTemp(Other)) {}
+	FORCEINLINE FNameDeprecationWrapper(const FString& Other) : Super(Other) {}
+	FORCEINLINE FNameDeprecationWrapper& operator=(FString&& Other) { Super::operator=(MoveTemp(Other)); return *this; }
+	FORCEINLINE FNameDeprecationWrapper& operator=(const FString& Other) { Super::operator=(Other); return *this; }
+
+	UE_DEPRECATED(5.5, "This variable is now an FString instead of an FName. Please update your code accordingly")
+	FORCEINLINE FNameDeprecationWrapper(FName&& DeprecatedValue) : Super(DeprecatedValue.ToString()) {}
+
+	UE_DEPRECATED(5.5, "This variable is now an FString instead of an FName. Please update your code accordingly")
+		FORCEINLINE FNameDeprecationWrapper(const FName& DeprecatedValue) : Super(DeprecatedValue.ToString()) {}
+
+	UE_DEPRECATED(5.5, "This variable is now an FString instead of an FName. Please update your code accordingly")
+		FORCEINLINE FNameDeprecationWrapper& operator=(FName&& DeprecatedValue) { Super::operator=(DeprecatedValue.ToString()); return *this; }
+
+	UE_DEPRECATED(5.5, "This variable is now an FString instead of an FName. Please update your code accordingly")
+		FORCEINLINE FNameDeprecationWrapper& operator=(const FName& DeprecatedValue) { Super::operator=(DeprecatedValue.ToString()); return *this; }
+
+	UE_DEPRECATED(5.5, "This variable is now an FString instead of an FName. Please update your code accordingly")
+		FORCEINLINE FString ToString() const { return *this; }
+
+	UE_DEPRECATED(5.5, "This variable is now an FString instead of an FName. Please update your code accordingly")
+		FORCEINLINE void ToString(FString& Out) const { Out = *this; }
+};
 
 /**
  * An interface used to collect and manage online stats
@@ -25,7 +194,7 @@ public:
 	 * @param StatName key name to search for
 	 * @return KeyValuePair if found, NULL otherwise
 	 */
-	class FVariantData* FindStatByName(const FName& StatName);
+	class FVariantData* FindStatByName(const FString& StatName);
 
 	/**
 	 * Sets a stat of type SDT_Float to the value specified. Does nothing
@@ -34,7 +203,7 @@ public:
 	 * @param StatName the stat to change the value of
 	 * @param Value the new value to assign to the stat
 	 */
-	virtual void SetFloatStat(const FName& StatName, float Value);
+	virtual void SetFloatStat(const FString& StatName, float Value);
 
 	/**
 	 * Sets a stat of type SDT_Int to the value specified. Does nothing
@@ -43,7 +212,7 @@ public:
 	 * @param StatName the stat to change the value of
 	 * @param Value the new value to assign to the stat
 	 */
-	virtual void SetIntStat(const FName& StatName, int32 Value);
+	virtual void SetIntStat(const FString& StatName, int32 Value);
 
 	/**
 	 * Increments a stat of type float by the value specified. Does nothing
@@ -52,7 +221,7 @@ public:
 	 * @param StatName the stat to increment
 	 * @param IncBy the value to increment by
 	 */
-	virtual void IncrementFloatStat(const FName& StatName, float IncBy = 1.0f);
+	virtual void IncrementFloatStat(const FString& StatName, float IncBy = 1.0f);
 
 	/**
 	 * Increments a stat of type int32 by the value specified. Does nothing
@@ -61,7 +230,7 @@ public:
 	 * @param StatName the stat to increment
 	 * @param IncBy the value to increment by
 	 */
-	virtual void IncrementIntStat(const FName& StatName, int32 IncBy = 1);
+	virtual void IncrementIntStat(const FString& StatName, int32 IncBy = 1);
 
 	/**
 	 * Decrements a stat of type float by the value specified. Does nothing
@@ -70,7 +239,7 @@ public:
 	 * @param StatName the stat to decrement
 	 * @param DecBy the value to decrement by
 	 */
-	virtual void DecrementFloatStat(const FName& StatName, float DecBy = 1.0f);
+	virtual void DecrementFloatStat(const FString& StatName, float DecBy = 1.0f);
 
 	/**
 	 * Decrements a stat of type int32 by the value specified. Does nothing
@@ -79,7 +248,22 @@ public:
 	 * @param StatName the stat to decrement
 	 * @param DecBy the value to decrement by
 	 */
-	virtual void DecrementIntStat(const FName& StatName, int32 DecBy = 1);
+	virtual void DecrementIntStat(const FString& StatName, int32 DecBy = 1);
+
+	UE_DEPRECATED(5.5, "Use the FString overload instead")
+	FVariantData* FindStatByName(const FName& StatName) { return FindStatByName(StatName.ToString()); }
+	UE_DEPRECATED(5.5, "Use the FString overload instead")
+	virtual void SetFloatStat(const FName& StatName, float Value) { SetFloatStat(StatName.ToString(), Value); }
+	UE_DEPRECATED(5.5, "Use the FString overload instead")
+	virtual void SetIntStat(const FName& StatName, int32 Value) { SetIntStat(StatName.ToString(), Value); }
+	UE_DEPRECATED(5.5, "Use the FString overload instead")
+	virtual void IncrementFloatStat(const FName& StatName, float IncBy = 1.0f) { IncrementFloatStat(StatName.ToString(), IncBy); }
+	UE_DEPRECATED(5.5, "Use the FString overload instead")
+	virtual void IncrementIntStat(const FName& StatName, int32 IncBy = 1) { IncrementIntStat(StatName.ToString(), IncBy); }
+	UE_DEPRECATED(5.5, "Use the FString overload instead")
+	virtual void DecrementFloatStat(const FName& StatName, float DecBy = 1.0f) { DecrementFloatStat(StatName.ToString(), DecBy); }
+	UE_DEPRECATED(5.5, "Use the FString overload instead")
+	virtual void DecrementIntStat(const FName& StatName, int32 DecBy = 1) { DecrementIntStat(StatName.ToString(), DecBy); }
 
 	/**
 	 * Destructor
@@ -105,16 +289,15 @@ public:
 	ELeaderboardUpdateMethod::Type UpdateMethod;
 
 	/** Names of the leaderboards to write to */
-	TArray<FName> LeaderboardNames;
+	FNameArrayDeprecationWrapper LeaderboardNames;
 
 	/** Name of the stat that the leaderboard is rated by */
-	FName RatedStat;
+	FNameDeprecationWrapper RatedStat;
 
 	FOnlineLeaderboardWrite() :
 		SortMethod(ELeaderboardSort::None),
 		DisplayFormat(ELeaderboardFormat::Number),
-		UpdateMethod(ELeaderboardUpdateMethod::KeepBest),
-		RatedStat(NAME_None)
+		UpdateMethod(ELeaderboardUpdateMethod::KeepBest)
 	{
 	}
 };
@@ -154,18 +337,24 @@ struct FColumnMetaData
 {
 private:
 	FColumnMetaData() :
-		ColumnName(NAME_None),
 		DataType(EOnlineKeyValuePairDataType::Empty)
 	{}
 
 public:
 
 	/** Name of the column to retrieve */
-	const FName ColumnName;
+	const FNameDeprecationWrapper ColumnName;
 	/** Type of data this column represents */
 	const EOnlineKeyValuePairDataType::Type DataType;
 
+	UE_DEPRECATED(5.5, "Use the FString overload instead")
 	FColumnMetaData(const FName InColumnName, EOnlineKeyValuePairDataType::Type InDataType) :
+		ColumnName(InColumnName.ToString()),
+		DataType(InDataType)
+	{
+	}
+
+	FColumnMetaData(const FString& InColumnName, EOnlineKeyValuePairDataType::Type InDataType) :
 		ColumnName(InColumnName),
 		DataType(InDataType)
 	{
@@ -179,9 +368,9 @@ class ONLINESUBSYSTEM_API FOnlineLeaderboardRead
 {
 public:
 	/** Name of the leaderboard read */
-	FName LeaderboardName;
+	FNameDeprecationWrapper LeaderboardName;
 	/** Column this leaderboard is sorted by */
-	FName SortedColumn;
+	FNameDeprecationWrapper SortedColumn;
 	/** Column metadata for this leaderboard */
 	TArray<FColumnMetaData> ColumnMetadata;
 	/** Array of ranked users retrieved (not necessarily sorted yet) */
@@ -190,8 +379,6 @@ public:
 	EOnlineAsyncTaskState::Type ReadState;
 
 	FOnlineLeaderboardRead() :
-		LeaderboardName(NAME_None),
-		SortedColumn(NAME_None),
 		ReadState(EOnlineAsyncTaskState::NotStarted)
 	{
 	}
diff --git a/Engine/Plugins/Online/OnlineSubsystem/Source/Test/OnlineAchievements.spec.cpp b/Engine/Plugins/Online/OnlineSubsystem/Source/Test/OnlineAchievements.spec.cpp
--- a/Engine/Plugins/Online/OnlineSubsystem/Source/Test/OnlineAchievements.spec.cpp
+++ b/Engine/Plugins/Online/OnlineSubsystem/Source/Test/OnlineAchievements.spec.cpp
@@ -86,7 +86,7 @@ void FOnlineAchievementsSpec::Define()
 
 										FOnlineAchievementsWritePtr AchievementWriteObject = MakeShareable(new FOnlineAchievementsWrite());
 										FOnlineAchievementsWriteRef AchievementWriter = AchievementWriteObject.ToSharedRef();
-										AchievementWriteObject->SetFloatStat(FName(*TestAchievement), 1.0f);
+										AchievementWriteObject->SetFloatStat(TestAchievement, 1.0f);
 
 										OnlineAchievements->WriteAchievements(*TestAccountId, AchievementWriter, FOnAchievementsWrittenDelegate::CreateLambda([this, &SomeAchievement, AchievementWriteObject, TestAccountId, TestAchievement, TestDone](const FUniqueNetId& WriteAchievementsPlayerId, bool bWriteAchievementsWasSuccessful)
 										{
diff --git a/Engine/Plugins/Online/OnlineSubsystem/Source/Test/OnlineLeaderboard.spec.cpp b/Engine/Plugins/Online/OnlineSubsystem/Source/Test/OnlineLeaderboard.spec.cpp
--- a/Engine/Plugins/Online/OnlineSubsystem/Source/Test/OnlineLeaderboard.spec.cpp
+++ b/Engine/Plugins/Online/OnlineSubsystem/Source/Test/OnlineLeaderboard.spec.cpp
@@ -386,17 +386,18 @@ void FOnlineLeaderboardSpec::Define()
 					{
 						UE_LOG_ONLINE_LEADERBOARD(Error, TEXT("OSS Automation: Test not yet implemented"));
 
-						FString TestAccountIdString = CommonUtils.GetSubsystemTestAccountUniqueId(SubsystemType);
+						const FString TestAccountIdString = CommonUtils.GetSubsystemTestAccountUniqueId(SubsystemType);
+						const FString TestStatString = TEXT("TestStat");
 						FUniqueNetIdPtr TestAccountId = OnlineIdentity->CreateUniquePlayerId(TestAccountIdString);
 
 						FOnlineLeaderboardWrite LeaderboardWriteObject;
-						LeaderboardWriteObject.LeaderboardNames.Add(FName(TEXT("TestLeaderboard")));
-						LeaderboardWriteObject.RatedStat = "TestStat";
+						LeaderboardWriteObject.LeaderboardNames.Add(FString(TEXT("TestLeaderboard")));
+						LeaderboardWriteObject.RatedStat = TestStatString;
 						LeaderboardWriteObject.DisplayFormat = ELeaderboardFormat::Number;
 						LeaderboardWriteObject.SortMethod = ELeaderboardSort::Descending;
 						LeaderboardWriteObject.UpdateMethod = ELeaderboardUpdateMethod::KeepBest;
 
-						LeaderboardWriteObject.SetIntStat("TestStat", 50);
+						LeaderboardWriteObject.SetIntStat(TestStatString, 50);
 
 						bool bCallStarted = OnlineLeaderboards->WriteLeaderboards(TEXT("TestSessionName"), *TestAccountId, LeaderboardWriteObject);
 						
@@ -409,17 +410,18 @@ void FOnlineLeaderboardSpec::Define()
 					{
 						UE_LOG_ONLINE_LEADERBOARD(Error, TEXT("OSS Automation: Test not yet implemented"));
 
-						FString TestAccountIdString = CommonUtils.GetSubsystemTestAccountUniqueId(SubsystemType);
+						const FString TestAccountIdString = CommonUtils.GetSubsystemTestAccountUniqueId(SubsystemType);
+						const FString TestStatString = TEXT("TestStat");
 						FUniqueNetIdPtr TestAccountId = OnlineIdentity->CreateUniquePlayerId(TestAccountIdString);
 
 						FOnlineLeaderboardWrite LeaderboardWriteObject;
-						LeaderboardWriteObject.LeaderboardNames.Add(FName(TEXT("TestLeaderboard")));
-						LeaderboardWriteObject.RatedStat = "TestStat";
+						LeaderboardWriteObject.LeaderboardNames.Add(FString(TEXT("TestLeaderboard")));
+						LeaderboardWriteObject.RatedStat = TestStatString;
 						LeaderboardWriteObject.DisplayFormat = ELeaderboardFormat::Number;
 						LeaderboardWriteObject.SortMethod = ELeaderboardSort::Descending;
 						LeaderboardWriteObject.UpdateMethod = ELeaderboardUpdateMethod::KeepBest;
 
-						LeaderboardWriteObject.SetIntStat("TestStat", 50);
+						LeaderboardWriteObject.SetIntStat(TestStatString, 50);
 
 						OnlineLeaderboards->WriteLeaderboards(TEXT(""), *TestAccountId, LeaderboardWriteObject);
 
@@ -430,16 +432,17 @@ void FOnlineLeaderboardSpec::Define()
 					{
 						UE_LOG_ONLINE_LEADERBOARD(Error, TEXT("OSS Automation: Test not yet implemented"));
 
+						const FString TestStatString = TEXT("TestStat");
 						FUniqueNetIdPtr TestAccountId = OnlineIdentity->CreateUniquePlayerId(TEXT("0123456789"));
 
 						FOnlineLeaderboardWrite LeaderboardWriteObject;
-						LeaderboardWriteObject.LeaderboardNames.Add(FName(TEXT("TestLeaderboard")));
-						LeaderboardWriteObject.RatedStat = "TestStat";
+						LeaderboardWriteObject.LeaderboardNames.Add(FString(TEXT("TestLeaderboard")));
+						LeaderboardWriteObject.RatedStat = TestStatString;
 						LeaderboardWriteObject.DisplayFormat = ELeaderboardFormat::Number;
 						LeaderboardWriteObject.SortMethod = ELeaderboardSort::Descending;
 						LeaderboardWriteObject.UpdateMethod = ELeaderboardUpdateMethod::KeepBest;
 
-						LeaderboardWriteObject.SetIntStat("TestStat", 50);
+						LeaderboardWriteObject.SetIntStat(TestStatString, 50);
 
 						OnlineLeaderboards->WriteLeaderboards(TEXT("TestSessionName"), *TestAccountId, LeaderboardWriteObject);
 
@@ -473,17 +476,18 @@ void FOnlineLeaderboardSpec::Define()
 					{
 						UE_LOG_ONLINE_LEADERBOARD(Error, TEXT("OSS Automation: Test not yet implemented"));
 
-						FString TestAccountIdString = CommonUtils.GetSubsystemTestAccountUniqueId(SubsystemType);
+						const FString TestAccountIdString = CommonUtils.GetSubsystemTestAccountUniqueId(SubsystemType);
+						const FString TestStatString = TEXT("TestStat");
 						FUniqueNetIdPtr TestAccountId = OnlineIdentity->CreateUniquePlayerId(TestAccountIdString);
 
 						FOnlineLeaderboardWrite LeaderboardWriteObject;
-						LeaderboardWriteObject.LeaderboardNames.Add(FName(TEXT("TestLeaderboard")));
-						LeaderboardWriteObject.RatedStat = "TestStat";
+						LeaderboardWriteObject.LeaderboardNames.Add(FString(TEXT("TestLeaderboard")));
+						LeaderboardWriteObject.RatedStat = TestStatString;
 						LeaderboardWriteObject.DisplayFormat = ELeaderboardFormat::Number;
 						LeaderboardWriteObject.SortMethod = ELeaderboardSort::Descending;
 						LeaderboardWriteObject.UpdateMethod = ELeaderboardUpdateMethod::KeepBest;
 
-						LeaderboardWriteObject.SetIntStat("TestStat", 50);
+						LeaderboardWriteObject.SetIntStat(TestStatString, 50);
 
 						OnlineLeaderboards->WriteLeaderboards(TEXT("TestSessionName"), *TestAccountId, LeaderboardWriteObject);
 
diff --git a/Engine/Plugins/Online/OnlineSubsystem/Source/Test/Utils/OnlineTestCommon.cpp b/Engine/Plugins/Online/OnlineSubsystem/Source/Test/Utils/OnlineTestCommon.cpp
--- a/Engine/Plugins/Online/OnlineSubsystem/Source/Test/Utils/OnlineTestCommon.cpp
+++ b/Engine/Plugins/Online/OnlineSubsystem/Source/Test/Utils/OnlineTestCommon.cpp
@@ -526,14 +526,14 @@ void FOnlineTestCommon::AddAchievementToTestAccount(IOnlineIdentityPtr OI, IOnli
 
 				if (PlayerAchievements.Num() > 0)
 				{
-					FString TestAchievement = PlayerAchievements[0].Id;
+					const FString TestAchievement = PlayerAchievements[0].Id;
 
 					FOnlineAchievement SomeAchievement;
 					OnlineAchievements->GetCachedAchievement(*TestAccountId, TestAchievement, SomeAchievement);
 
 					FOnlineAchievementsWritePtr AchievementWriteObject = MakeShareable(new FOnlineAchievementsWrite());
 					FOnlineAchievementsWriteRef AchievementWriter = AchievementWriteObject.ToSharedRef();
-					AchievementWriteObject->SetFloatStat(FName(*TestAchievement), 1.0f);
+					AchievementWriteObject->SetFloatStat(TestAchievement, 1.0f);
 
 					OnlineAchievements->WriteAchievements(*TestAccountId, AchievementWriter, FOnAchievementsWrittenDelegate::CreateLambda([this, AchievementWriteObject, TestAccountId, TestAchievement, TestDone](const FUniqueNetId& WriteAchievements, bool bWriteAchievementsWasSuccessful)
 					{
diff --git a/Engine/Plugins/Online/OnlineSubsystemEOS/Source/OnlineSubsystemEOS/Private/OnlineAchievementsEOS.cpp b/Engine/Plugins/Online/OnlineSubsystemEOS/Source/OnlineSubsystemEOS/Private/OnlineAchievementsEOS.cpp
--- a/Engine/Plugins/Online/OnlineSubsystemEOS/Source/OnlineSubsystemEOS/Private/OnlineAchievementsEOS.cpp
+++ b/Engine/Plugins/Online/OnlineSubsystemEOS/Source/OnlineSubsystemEOS/Private/OnlineAchievementsEOS.cpp
@@ -34,15 +34,15 @@ void FOnlineAchievementsEOS::UnlockAchievements(const FUniqueNetId& PlayerId, FO
 		return;
 	}
 
-	TArray<FName> InAchievementIds;
+	TArray<FString> InAchievementIds;
 	WriteObject->Properties.GenerateKeyArray(InAchievementIds);
 	TArray<FTCHARToUTF8> AchievementIdConverters; // We can't use StringCast<UTF8CHAR> because it's non-copyable, and the array will make a copy
 	AchievementIdConverters.Reserve(InAchievementIds.Num());
 	TArray<const char*> AchievementIdPtrs;
 	AchievementIdPtrs.Reserve(InAchievementIds.Num());
-	for (const FName& AchievementId : InAchievementIds)
+	for (const FString& AchievementId : InAchievementIds)
 	{
-		const FTCHARToUTF8& Converter = AchievementIdConverters.Emplace_GetRef(*AchievementId.ToString());
+		const FTCHARToUTF8& Converter = AchievementIdConverters.Emplace_GetRef(*AchievementId);
 		AchievementIdPtrs.Emplace(Converter.Get());
 	}
 
@@ -63,14 +63,14 @@ void FOnlineAchievementsEOS::UnlockAchievements(const FUniqueNetId& PlayerId, FO
 			{
 				UE_LOG_ONLINE_ACHIEVEMENTS(Verbose, TEXT("(%d) achievements unlocked. Caching achievements for user (%s)"), (uint32)Data->AchievementsCount, *LambdaPlayerId->ToString());
 
-				for (const FName& AchievementId : InAchievementIds)
+				for (const FString& AchievementId : InAchievementIds)
 				{
 					for (FOnlineAchievement& Achievement : Achievements->Get())
 					{
-						if (Achievement.Id == AchievementId.ToString())
+						if (Achievement.Id == AchievementId)
 						{
 							Achievement.Progress = 1.0f;
-							TriggerOnAchievementUnlockedDelegates(*LambdaPlayerId, AchievementId.ToString());
+							TriggerOnAchievementUnlockedDelegates(*LambdaPlayerId, AchievementId);
 							break;
 						}
 					}
diff --git a/Engine/Plugins/Online/OnlineSubsystemEOS/Source/OnlineSubsystemEOS/Private/OnlineLeaderboardsEOS.cpp b/Engine/Plugins/Online/OnlineSubsystemEOS/Source/OnlineSubsystemEOS/Private/OnlineLeaderboardsEOS.cpp
--- a/Engine/Plugins/Online/OnlineSubsystemEOS/Source/OnlineSubsystemEOS/Private/OnlineLeaderboardsEOS.cpp
+++ b/Engine/Plugins/Online/OnlineSubsystemEOS/Source/OnlineSubsystemEOS/Private/OnlineLeaderboardsEOS.cpp
@@ -107,7 +107,7 @@ bool FOnlineLeaderboardsEOS::ReadLeaderboards(const TArray<FUniqueNetIdRef>& Pla
 	int32 Index = 0;
 	for (const FColumnMetaData& Column : ReadObject->ColumnMetadata)
 	{
-		FCStringAnsi::Strncpy(Options.PointerArray[Index], TCHAR_TO_UTF8(*Column.ColumnName.ToString()), EOS_OSS_STRING_BUFFER_LENGTH);
+		FCStringAnsi::Strncpy(Options.PointerArray[Index], TCHAR_TO_UTF8(*Column.ColumnName), EOS_OSS_STRING_BUFFER_LENGTH);
 		EOS_Leaderboards_UserScoresQueryStatInfo& StatInfo = Options.StatInfoArray[Index];
 		StatInfo.ApiVersion = 1;
 		UE_EOS_CHECK_API_MISMATCH(EOS_LEADERBOARDS_USERSCORESQUERYSTATINFO_API_LATEST, 1);
@@ -156,7 +156,7 @@ bool FOnlineLeaderboardsEOS::ReadLeaderboards(const TArray<FUniqueNetIdRef>& Pla
 			for (const FColumnMetaData& Column : QueryContext->ReadObject->ColumnMetadata)
 			{
 				// Update which stat we are requesting
-				FCStringAnsi::Strncpy(StatName, TCHAR_TO_UTF8(*Column.ColumnName.ToString()), EOS_OSS_STRING_BUFFER_LENGTH);
+				FCStringAnsi::Strncpy(StatName, TCHAR_TO_UTF8(*Column.ColumnName), EOS_OSS_STRING_BUFFER_LENGTH);
 
 				EOS_Leaderboards_LeaderboardUserScore* LeaderboardUserScore = nullptr;
 				EOS_EResult UserCopyResult = EOS_Leaderboards_CopyLeaderboardUserScoreByUserId(EOSSubsystem->LeaderboardsHandle, &UserCopyOptions, &LeaderboardUserScore);
@@ -174,8 +174,7 @@ bool FOnlineLeaderboardsEOS::ReadLeaderboards(const TArray<FUniqueNetIdRef>& Pla
 		}
 
 		// Manually build the ranks by sorting and then assigning rank values
-		FName SortedColumn = QueryContext->ReadObject->SortedColumn;
-		QueryContext->ReadObject->Rows.Sort([SortedColumn](const FOnlineStatsRow& RowA, const FOnlineStatsRow& RowB)
+		QueryContext->ReadObject->Rows.Sort([SortedColumn = QueryContext->ReadObject->SortedColumn](const FOnlineStatsRow& RowA, const FOnlineStatsRow& RowB)
 		{
 			const FVariantData& ValueA = RowA.Columns[SortedColumn];
 			const FVariantData& ValueB = RowB.Columns[SortedColumn];
@@ -255,7 +254,7 @@ bool FOnlineLeaderboardsEOS::ReadLeaderboardsAroundRank(int32 Rank, uint32 Range
 	UE_EOS_CHECK_API_MISMATCH(EOS_LEADERBOARDS_QUERYLEADERBOARDRANKS_API_LATEST, 2);
 	Options.LeaderboardId = LeaderboardId;
 	Options.LocalUserId = EOSSubsystem->UserManager->GetLocalProductUserId(0);
-	FCStringAnsi::Strncpy(LeaderboardId, TCHAR_TO_UTF8(*ReadObject->LeaderboardName.ToString()), EOS_OSS_STRING_BUFFER_LENGTH);
+	FCStringAnsi::Strncpy(LeaderboardId, TCHAR_TO_UTF8(*ReadObject->LeaderboardName), EOS_OSS_STRING_BUFFER_LENGTH);
 
 	FQueryLeaderboardCallback* CallbackObj = new FQueryLeaderboardCallback(FOnlineLeaderboardsEOSWeakPtr(AsShared()));
 	FOnlineLeaderboardReadRef LambdaReadObject = ReadObject;
@@ -354,9 +353,9 @@ bool FOnlineLeaderboardsEOS::WriteLeaderboards(const FName& SessionName, const F
 	TArray<FOnlineStatsUserUpdatedStats> StatsToWrite;
 
 	FOnlineStatsUserUpdatedStats& UpdatedStats = StatsToWrite.Emplace_GetRef(Player.AsShared());
-	for (const TPair<FName, FVariantData>& Stat : WriteObject.Properties)
+	for (const TPair<FString, FVariantData>& Stat : WriteObject.Properties)
 	{
-		UpdatedStats.Stats.Add(Stat.Key.ToString(), FOnlineStatUpdate(Stat.Value, FOnlineStatUpdate::EOnlineStatModificationType::Unknown));
+		UpdatedStats.Stats.Add(Stat.Key, FOnlineStatUpdate(Stat.Value, FOnlineStatUpdate::EOnlineStatModificationType::Unknown));
 	}
 
 	EOSSubsystem->StatsInterfacePtr->UpdateStats(Player.AsShared(), StatsToWrite, FOnlineStatsUpdateStatsComplete());
@@ -400,9 +399,9 @@ bool FOnlineLeaderboardsEOS::HandleLeaderboardsExec(UWorld* InWorld, const TCHAR
 		const EOnlineKeyValuePairDataType::Type DataType = EOnlineKeyValuePairDataType::FromString(DataTypeStr);
 
 		FOnlineLeaderboardReadRef ReadRef = MakeShared<FOnlineLeaderboardRead>();
-		ReadRef->LeaderboardName = FName(*LeaderboardName);
-		ReadRef->SortedColumn = FName(*SortedColumn);
-		ReadRef->ColumnMetadata.Add(FColumnMetaData(FName(*SortedColumn), DataType));
+		ReadRef->LeaderboardName = LeaderboardName;
+		ReadRef->SortedColumn = SortedColumn;
+		ReadRef->ColumnMetadata.Add(FColumnMetaData(SortedColumn, DataType));
 
 		AddOnLeaderboardReadCompleteDelegate_Handle(FOnLeaderboardReadCompleteDelegate::CreateLambda([ReadRef](bool bWasSuccessful)
 			{
@@ -426,9 +425,9 @@ bool FOnlineLeaderboardsEOS::HandleLeaderboardsExec(UWorld* InWorld, const TCHAR
 		const EOnlineKeyValuePairDataType::Type DataType = EOnlineKeyValuePairDataType::FromString(DataTypeStr);
 
 		FOnlineLeaderboardReadRef ReadRef = MakeShared<FOnlineLeaderboardRead>();
-		ReadRef->LeaderboardName = FName(*LeaderboardName);
-		ReadRef->SortedColumn = FName(*SortedColumn);
-		ReadRef->ColumnMetadata.Add(FColumnMetaData(FName(*SortedColumn), DataType));
+		ReadRef->LeaderboardName = LeaderboardName;
+		ReadRef->SortedColumn = SortedColumn;
+		ReadRef->ColumnMetadata.Add(FColumnMetaData(SortedColumn, DataType));
 
 		AddOnLeaderboardReadCompleteDelegate_Handle(FOnLeaderboardReadCompleteDelegate::CreateLambda([ReadRef](bool bWasSuccessful)
 			{
@@ -454,9 +453,9 @@ bool FOnlineLeaderboardsEOS::HandleLeaderboardsExec(UWorld* InWorld, const TCHAR
 		const EOnlineKeyValuePairDataType::Type DataType = EOnlineKeyValuePairDataType::FromString(DataTypeStr);
 
 		FOnlineLeaderboardReadRef ReadRef = MakeShared<FOnlineLeaderboardRead>();
-		ReadRef->LeaderboardName = FName(*LeaderboardName);
-		ReadRef->SortedColumn = FName(*SortedColumn);
-		ReadRef->ColumnMetadata.Add(FColumnMetaData(FName(*SortedColumn), DataType));
+		ReadRef->LeaderboardName = LeaderboardName;
+		ReadRef->SortedColumn = SortedColumn;
+		ReadRef->ColumnMetadata.Add(FColumnMetaData(SortedColumn, DataType));
 
 		AddOnLeaderboardReadCompleteDelegate_Handle(FOnLeaderboardReadCompleteDelegate::CreateLambda([ReadRef](bool bWasSuccessful)
 			{
@@ -484,9 +483,9 @@ bool FOnlineLeaderboardsEOS::HandleLeaderboardsExec(UWorld* InWorld, const TCHAR
 		const EOnlineKeyValuePairDataType::Type DataType = EOnlineKeyValuePairDataType::FromString(DataTypeStr);
 
 		FOnlineLeaderboardReadRef ReadRef = MakeShared<FOnlineLeaderboardRead>();
-		ReadRef->LeaderboardName = FName(*LeaderboardName);
-		ReadRef->SortedColumn = FName(*SortedColumn);
-		ReadRef->ColumnMetadata.Add(FColumnMetaData(FName(*SortedColumn), DataType));
+		ReadRef->LeaderboardName = LeaderboardName;
+		ReadRef->SortedColumn = SortedColumn;
+		ReadRef->ColumnMetadata.Add(FColumnMetaData(SortedColumn, DataType));
 
 		AddOnLeaderboardReadCompleteDelegate_Handle(FOnLeaderboardReadCompleteDelegate::CreateLambda([ReadRef](bool bWasSuccessful)
 			{
diff --git a/Engine/Plugins/Online/OnlineSubsystemNull/Source/Private/OnlineAchievementsInterfaceNull.cpp b/Engine/Plugins/Online/OnlineSubsystemNull/Source/Private/OnlineAchievementsInterfaceNull.cpp
--- a/Engine/Plugins/Online/OnlineSubsystemNull/Source/Private/OnlineAchievementsInterfaceNull.cpp
+++ b/Engine/Plugins/Online/OnlineSubsystemNull/Source/Private/OnlineAchievementsInterfaceNull.cpp
@@ -44,7 +44,7 @@ void FOnlineAchievementsNull::WriteAchievements(const FUniqueNetId& PlayerId, FO
 	const int32 AchNum = PlayerAch->Num();
 	for (FStatPropertyArray::TConstIterator It(WriteObject->Properties); It; ++It)
 	{
-		const FString AchievementId = It.Key().ToString();
+		const FString AchievementId = It.Key();
 		for (int32 AchIdx = 0; AchIdx < AchNum; ++AchIdx)
 		{
 			if ((*PlayerAch)[ AchIdx ].Id == AchievementId)
diff --git a/Engine/Plugins/Online/OnlineSubsystemNull/Source/Private/OnlineLeaderboardInterfaceNull.cpp b/Engine/Plugins/Online/OnlineSubsystemNull/Source/Private/OnlineLeaderboardInterfaceNull.cpp
--- a/Engine/Plugins/Online/OnlineSubsystemNull/Source/Private/OnlineLeaderboardInterfaceNull.cpp
+++ b/Engine/Plugins/Online/OnlineSubsystemNull/Source/Private/OnlineLeaderboardInterfaceNull.cpp
@@ -112,7 +112,7 @@ bool FOnlineLeaderboardsNull::WriteLeaderboards(const FName& SessionName, const
 
 		for (FStatPropertyArray::TConstIterator It(WriteObject.Properties); It; ++It)
 		{
-			const FName& StatName = It.Key();
+			const FString& StatName = It.Key();
 			const FVariantData& Stat = It.Value();
 			FVariantData* ExistingStat = PlayerRow->Columns.Find(StatName);
 			if (ExistingStat)
@@ -156,7 +156,7 @@ bool FOnlineLeaderboardsNull::WriteLeaderboards(const FName& SessionName, const
 	return bWasSuccessful;
 }
 
-FOnlineLeaderboardsNull::FLeaderboardNull* FOnlineLeaderboardsNull::FindOrCreateLeaderboard(const FName& LeaderboardName, ELeaderboardSort::Type SortMethod, ELeaderboardFormat::Type DisplayFormat)
+FOnlineLeaderboardsNull::FLeaderboardNull* FOnlineLeaderboardsNull::FindOrCreateLeaderboard(const FString& LeaderboardName, ELeaderboardSort::Type SortMethod, ELeaderboardFormat::Type DisplayFormat)
 {
 	FLeaderboardNull* Existing = Leaderboards.Find(LeaderboardName);
 	if (Existing == NULL)
diff --git a/Engine/Plugins/Online/OnlineSubsystemNull/Source/Private/OnlineLeaderboardInterfaceNull.h b/Engine/Plugins/Online/OnlineSubsystemNull/Source/Private/OnlineLeaderboardInterfaceNull.h
--- a/Engine/Plugins/Online/OnlineSubsystemNull/Source/Private/OnlineLeaderboardInterfaceNull.h
+++ b/Engine/Plugins/Online/OnlineSubsystemNull/Source/Private/OnlineLeaderboardInterfaceNull.h
@@ -43,7 +43,7 @@ private:
 	class FOnlineSubsystemNull* NullSubsystem;
 
 	/** Leaderboards maintained by the subsystem */
-	TMap<FName, FLeaderboardNull> Leaderboards;
+	TMap<FString, FLeaderboardNull> Leaderboards;
 
 	FOnlineLeaderboardsNull() : 
 		NullSubsystem(NULL)
@@ -58,7 +58,7 @@ private:
 	 * @param SortMethod method the leaderboard scores will be sorted, ignored if leaderboard exists
 	 * @param DisplayFormat type of data the leaderboard represents, ignored if leaderboard exists
 	 */
-	FLeaderboardNull* FindOrCreateLeaderboard(const FName& LeaderboardName, ELeaderboardSort::Type SortMethod, ELeaderboardFormat::Type DisplayFormat);
+	FLeaderboardNull* FindOrCreateLeaderboard(const FString& LeaderboardName, ELeaderboardSort::Type SortMethod, ELeaderboardFormat::Type DisplayFormat);
 
 PACKAGE_SCOPE:
 
diff --git a/Engine/Plugins/Online/OnlineSubsystemOculus/Source/Private/OnlineAchievementsInterfaceOculus.cpp b/Engine/Plugins/Online/OnlineSubsystemOculus/Source/Private/OnlineAchievementsInterfaceOculus.cpp
--- a/Engine/Plugins/Online/OnlineSubsystemOculus/Source/Private/OnlineAchievementsInterfaceOculus.cpp
+++ b/Engine/Plugins/Online/OnlineSubsystemOculus/Source/Private/OnlineAchievementsInterfaceOculus.cpp
@@ -103,7 +103,7 @@ PRAGMA_ENABLE_DEPRECATION_WARNINGS
 	// treat each achievement as unlocked
 	for (FStatPropertyArray::TConstIterator It(WriteObject->Properties); It; ++It)
 	{
-		const FString AchievementId = It.Key().ToString();
+		const FString AchievementId = It.Key();
 		auto VariantData = It.Value();
 
 		auto AchievementDesc = AchievementDescriptions.Find(AchievementId);
diff --git a/Engine/Plugins/Online/OnlineSubsystemOculus/Source/Private/OnlineLeaderboardInterfaceOculus.cpp b/Engine/Plugins/Online/OnlineSubsystemOculus/Source/Private/OnlineLeaderboardInterfaceOculus.cpp
--- a/Engine/Plugins/Online/OnlineSubsystemOculus/Source/Private/OnlineLeaderboardInterfaceOculus.cpp
+++ b/Engine/Plugins/Online/OnlineSubsystemOculus/Source/Private/OnlineLeaderboardInterfaceOculus.cpp
@@ -66,7 +66,7 @@ bool FOnlineLeaderboardOculus::ReadOculusLeaderboards(bool bOnlyFriends, bool bO
 	ReadObject->ReadState = EOnlineAsyncTaskState::InProgress;
 PRAGMA_DISABLE_DEPRECATION_WARNINGS
 	OculusSubsystem.AddRequestDelegate(
-		ovr_Leaderboard_GetEntries(TCHAR_TO_ANSI(*ReadObject->LeaderboardName.ToString()), Limit, FilterType, StartAt),
+		ovr_Leaderboard_GetEntries(TCHAR_TO_ANSI(*ReadObject->LeaderboardName), Limit, FilterType, StartAt),
 		FOculusMessageOnCompleteDelegate::CreateLambda([this, ReadObject](ovrMessageHandle Message, bool bIsError)
 	{
 		OnReadLeaderboardsComplete(Message, bIsError, ReadObject);
@@ -192,7 +192,7 @@ PRAGMA_ENABLE_DEPRECATION_WARNINGS
 
 	if (StatData == nullptr)
 	{
-		UE_LOG_ONLINE_LEADERBOARD(Error, TEXT("Could not find RatedStat: %s"), *WriteObject.RatedStat.ToString());
+		UE_LOG_ONLINE_LEADERBOARD(Error, TEXT("Could not find RatedStat: %s"), *WriteObject.RatedStat);
 		return false;
 	}
 
@@ -235,7 +235,7 @@ PRAGMA_ENABLE_DEPRECATION_WARNINGS
 	{
 PRAGMA_DISABLE_DEPRECATION_WARNINGS
 		OculusSubsystem.AddRequestDelegate(
-			ovr_Leaderboard_WriteEntry(TCHAR_TO_ANSI(*LeaderboardName.ToString()), Score, /* extra_data */ nullptr, 0, (WriteObject.UpdateMethod == ELeaderboardUpdateMethod::Force)),
+			ovr_Leaderboard_WriteEntry(TCHAR_TO_ANSI(*LeaderboardName), Score, /* extra_data */ nullptr, 0, (WriteObject.UpdateMethod == ELeaderboardUpdateMethod::Force)),
 			FOculusMessageOnCompleteDelegate::CreateLambda([this](ovrMessageHandle Message, bool bIsError)
 		{
 			if (bIsError) 
diff --git a/Engine/Plugins/Online/OnlineSubsystemSteam/Source/Private/OnlineAchievementsInterfaceSteam.cpp b/Engine/Plugins/Online/OnlineSubsystemSteam/Source/Private/OnlineAchievementsInterfaceSteam.cpp
--- a/Engine/Plugins/Online/OnlineSubsystemSteam/Source/Private/OnlineAchievementsInterfaceSteam.cpp
+++ b/Engine/Plugins/Online/OnlineSubsystemSteam/Source/Private/OnlineAchievementsInterfaceSteam.cpp
@@ -67,7 +67,7 @@ void FOnlineAchievementsSteam::WriteAchievements(const FUniqueNetId& PlayerId, F
 	const int32 AchNum = PlayerAch->Num();
 	for (FStatPropertyArray::TConstIterator It(WriteObject->Properties); It; ++It)
 	{
-		const FString AchievementId = It.Key().ToString();
+		const FString AchievementId = It.Key();
 		UE_LOG_ONLINE_ACHIEVEMENTS(Verbose, TEXT("WriteObject AchievementId: '%s'"), *AchievementId);
 		for (int32 AchIdx = 0; AchIdx < AchNum; ++AchIdx)
 		{
@@ -117,7 +117,7 @@ void FOnlineAchievementsSteam::OnWriteAchievementsComplete(const FUniqueNetIdSte
 			const int32 AchNum = PlayerAch->Num();
 			for (FStatPropertyArray::TConstIterator It(WriteObject->Properties); It; ++It)
 			{
-				const FString AchievementId = It.Key().ToString();
+				const FString AchievementId = It.Key();
 				for (int32 AchIdx = 0; AchIdx < AchNum; ++AchIdx)
 				{
 					if ((*PlayerAch)[ AchIdx ].Id == AchievementId)
diff --git a/Engine/Plugins/Online/OnlineSubsystemSteam/Source/Private/OnlineLeaderboardInterfaceSteam.cpp b/Engine/Plugins/Online/OnlineSubsystemSteam/Source/Private/OnlineLeaderboardInterfaceSteam.cpp
--- a/Engine/Plugins/Online/OnlineSubsystemSteam/Source/Private/OnlineLeaderboardInterfaceSteam.cpp
+++ b/Engine/Plugins/Online/OnlineSubsystemSteam/Source/Private/OnlineLeaderboardInterfaceSteam.cpp
@@ -13,9 +13,9 @@
  * @param LeaderboardName name of leaderboard 
  * @param StatName name of stat
  */
-inline FName GetLeaderboardStatName(const FName& LeaderboardName, const FName& StatName)
+inline FString GetLeaderboardStatName(const FString& LeaderboardName, const FString& StatName)
 {
-	return TCHAR_TO_ANSI((*FString::Printf(TEXT("%s_%s"), *LeaderboardName.ToString(), *StatName.ToString())));
+	return FString::Printf(TEXT("%s_%s"), *LeaderboardName, *StatName);
 }
 
 /** Helper function to convert enums */
@@ -265,7 +265,7 @@ public:
 					for (FStatPropertyArray::TConstIterator It(Stats); It; ++It)
 					{
 						bool bSuccess = false;
-						const FString StatName = It.Key().ToString();
+						const FString StatName = It.Key();
 						const FVariantData& Stat = It.Value();
 
 						switch (Stat.GetType())
@@ -419,8 +419,8 @@ public:
 				for (int32 StatIdx = 0; StatIdx < ReadObject->ColumnMetadata.Num(); StatIdx++)
 				{
 					const FColumnMetaData& ColumnMeta = ReadObject->ColumnMetadata[StatIdx];
-					FName LeaderboardStat = GetLeaderboardStatName(ReadObject->LeaderboardName, ColumnMeta.ColumnName);
-					const FString StatName = LeaderboardStat.ToString();
+					const FString LeaderboardStat = GetLeaderboardStatName(ReadObject->LeaderboardName, ColumnMeta.ColumnName);
+					const FString StatName = LeaderboardStat;
 
 					bool bSuccess = false;
 					FVariantData* LastColumn = NULL;
@@ -592,7 +592,7 @@ private:
 	/** Has this request been started */
 	bool bInit;
 	/** Name of requested leaderboard */
-	FName LeaderboardName;
+	FString LeaderboardName;
 	/** Method of sorting the scores on the leaderboard */
 	ELeaderboardSort::Type SortMethod;
 	/** Method of displaying the data on the leaderboard */
@@ -609,17 +609,17 @@ private:
 	 * @param SortMethod method the leaderboard scores will be sorted, ignored if leaderboard exists
 	 * @param DisplayFormat type of data the leaderboard represents, ignored if leaderboard exists
 	 */
-	void CreateOrFindLeaderboard(const FName& InLeaderboardName, ELeaderboardSort::Type InSortMethod, ELeaderboardFormat::Type InDisplayFormat)
+	void CreateOrFindLeaderboard(const FString& InLeaderboardName, ELeaderboardSort::Type InSortMethod, ELeaderboardFormat::Type InDisplayFormat)
 	{
 		if (bFindOnly)
 		{
-			CallbackHandle = SteamUserStats()->FindLeaderboard(TCHAR_TO_UTF8(*InLeaderboardName.ToString()));
+			CallbackHandle = SteamUserStats()->FindLeaderboard(TCHAR_TO_UTF8(*InLeaderboardName));
 		}
 		else
 		{
 			ELeaderboardSortMethod SortMethodSteam = ToSteamLeaderboardSortMethod(InSortMethod);
 			ELeaderboardDisplayType DisplayTypeSteam = ToSteamLeaderboardDisplayType(InDisplayFormat);
-			CallbackHandle = SteamUserStats()->FindOrCreateLeaderboard(TCHAR_TO_UTF8(*InLeaderboardName.ToString()), SortMethodSteam, DisplayTypeSteam);
+			CallbackHandle = SteamUserStats()->FindOrCreateLeaderboard(TCHAR_TO_UTF8(*InLeaderboardName), SortMethodSteam, DisplayTypeSteam);
 		}
 	}
 
@@ -627,7 +627,6 @@ private:
 	FOnlineAsyncTaskSteamRetrieveLeaderboard() :
 		FOnlineAsyncTaskSteam(NULL, k_uAPICallInvalid),
 		bInit(false),
-		LeaderboardName(NAME_None),
 		SortMethod(ELeaderboardSort::Ascending),
 		DisplayFormat(ELeaderboardFormat::Number),
 		bFindOnly(true)
@@ -636,7 +635,7 @@ private:
 
 public:
 	/** Create a leaderboard implementation */
-	FOnlineAsyncTaskSteamRetrieveLeaderboard(FOnlineSubsystemSteam* InSteamSubsystem, const FName& InLeaderboardName, ELeaderboardSort::Type InSortMethod, ELeaderboardFormat::Type InDisplayFormat) :
+	FOnlineAsyncTaskSteamRetrieveLeaderboard(FOnlineSubsystemSteam* InSteamSubsystem, const FString& InLeaderboardName, ELeaderboardSort::Type InSortMethod, ELeaderboardFormat::Type InDisplayFormat) :
 		FOnlineAsyncTaskSteam(InSteamSubsystem, k_uAPICallInvalid),
 		bInit(false),
 		LeaderboardName(InLeaderboardName),
@@ -647,7 +646,7 @@ public:
 	}
 
 	/** Find a leaderboard implementation */
-	FOnlineAsyncTaskSteamRetrieveLeaderboard(FOnlineSubsystemSteam* InSteamSubsystem, const FName& InLeaderboardName) :
+	FOnlineAsyncTaskSteamRetrieveLeaderboard(FOnlineSubsystemSteam* InSteamSubsystem, const FString& InLeaderboardName) :
 		FOnlineAsyncTaskSteam(InSteamSubsystem, k_uAPICallInvalid),
 		bInit(false),
 		LeaderboardName(InLeaderboardName),
@@ -723,7 +722,7 @@ public:
 		if (bWasSuccessful)
 		{
 			ISteamUserStats* SteamUserPtr = SteamUserStats();
-			check(LeaderboardName.ToString() == FString(SteamUserPtr->GetLeaderboardName(CallbackResults.m_hSteamLeaderboard)));
+			check(LeaderboardName == FString(SteamUserPtr->GetLeaderboardName(CallbackResults.m_hSteamLeaderboard)));
 
 			Leaderboard->LeaderboardHandle = CallbackResults.m_hSteamLeaderboard;
 			Leaderboard->TotalLeaderboardRows = SteamUserPtr->GetLeaderboardEntryCount(CallbackResults.m_hSteamLeaderboard);
@@ -1074,9 +1073,9 @@ private:
 	/** Has this request been started */
 	bool bInit;
 	/** Name of leaderboard to update */
-	FName LeaderboardName;
+	FString LeaderboardName;
 	/** Name of stat that will replace/update the existing value on the leaderboard */
-	FName RatedStat;
+	FString RatedStat;
 	/** Score that will replace/update the existing value on the leaderboard */
 	int32 NewScore;
 	/** Method of update against the previous score */
@@ -1089,8 +1088,6 @@ private:
 	FOnlineAsyncTaskSteamUpdateLeaderboard() :
 		FOnlineAsyncTaskSteam(NULL, k_uAPICallInvalid),
 		bInit(false),
-		LeaderboardName(NAME_None),
-		RatedStat(NAME_None),
 		NewScore(0),
 		UpdateMethod(ELeaderboardUpdateMethod::KeepBest),
 		bShouldTriggerDelegates(false)
@@ -1098,7 +1095,7 @@ private:
 	}
 
 public:
-	FOnlineAsyncTaskSteamUpdateLeaderboard(FOnlineSubsystemSteam* InSteamSubsystem, const FName& InLeaderboardName, const FName& InRatedStat, ELeaderboardUpdateMethod::Type InUpdateMethod, bool bInShouldTriggerDelegates) :
+	FOnlineAsyncTaskSteamUpdateLeaderboard(FOnlineSubsystemSteam* InSteamSubsystem, const FString& InLeaderboardName, const FString& InRatedStat, ELeaderboardUpdateMethod::Type InUpdateMethod, bool bInShouldTriggerDelegates) :
 		FOnlineAsyncTaskSteam(InSteamSubsystem, k_uAPICallInvalid),
 		bInit(false),
 		LeaderboardName(InLeaderboardName),
@@ -1114,7 +1111,7 @@ public:
 	 */
 	virtual FString ToString() const override
 	{
-		return FString::Printf(TEXT("FOnlineAsyncTaskSteamUpdateLeaderboard bWasSuccessful: %d Leaderboard: %s Score: %d"), WasSuccessful(), *LeaderboardName.ToString(), NewScore);
+		return FString::Printf(TEXT("FOnlineAsyncTaskSteamUpdateLeaderboard bWasSuccessful: %d Leaderboard: %s Score: %d"), WasSuccessful(), *LeaderboardName, NewScore);
 	}
 
 	/**
@@ -1157,7 +1154,7 @@ public:
 					break;
 				}
 
-				const FString RatedStatName = GetLeaderboardStatName(LeaderboardName, RatedStat).ToString();
+				const FString RatedStatName = GetLeaderboardStatName(LeaderboardName, RatedStat);
 				if (SteamUserStats()->GetStat(TCHAR_TO_UTF8(*RatedStatName), &NewScore))
 				{
 					CallbackHandle = SteamUserStatsPtr->UploadLeaderboardScore(LeaderboardHandle, UpdateMethodSteam, NewScore, NULL, 0);
@@ -1505,7 +1502,7 @@ bool FOnlineLeaderboardsSteam::WriteLeaderboards(const FName& SessionName, const
 		for (FStatPropertyArray::TConstIterator It(WriteObject.Properties); It; ++It)
 		{
 			const FVariantData& Stat = It.Value();
-			FName LeaderboardStatName = GetLeaderboardStatName(WriteObject.LeaderboardNames[LeaderboardIdx], It.Key());
+			const FString LeaderboardStatName = GetLeaderboardStatName(WriteObject.LeaderboardNames[LeaderboardIdx], It.Key());
 			LeaderboardStats.Add(LeaderboardStatName, Stat);
 		}
 	}
@@ -1546,7 +1543,7 @@ bool FOnlineLeaderboardsSteam::WriteOnlinePlayerRatings(const FName& SessionName
 	return false;
 }
 
-FLeaderboardMetadataSteam* FOnlineLeaderboardsSteam::GetLeaderboardMetadata(const FName& LeaderboardName)
+FLeaderboardMetadataSteam* FOnlineLeaderboardsSteam::GetLeaderboardMetadata(const FString& LeaderboardName)
 {
 	FScopeLock ScopeLock(&LeaderboardMetadataLock);
 	for (int32 LeaderboardIdx = 0; LeaderboardIdx < Leaderboards.Num(); LeaderboardIdx++)
@@ -1561,7 +1558,7 @@ FLeaderboardMetadataSteam* FOnlineLeaderboardsSteam::GetLeaderboardMetadata(cons
 	return NULL;
 }
 
-void FOnlineLeaderboardsSteam::CreateLeaderboard(const FName& LeaderboardName, ELeaderboardSort::Type SortMethod, ELeaderboardFormat::Type DisplayFormat)
+void FOnlineLeaderboardsSteam::CreateLeaderboard(const FString& LeaderboardName, ELeaderboardSort::Type SortMethod, ELeaderboardFormat::Type DisplayFormat)
 {
 	FScopeLock ScopeLock(&LeaderboardMetadataLock);
 	FLeaderboardMetadataSteam* LeaderboardMetadata = GetLeaderboardMetadata(LeaderboardName);
@@ -1580,7 +1577,7 @@ void FOnlineLeaderboardsSteam::CreateLeaderboard(const FName& LeaderboardName, E
 	// else request already in flight or already found
 }
 
-void FOnlineLeaderboardsSteam::FindLeaderboard(const FName& LeaderboardName)
+void FOnlineLeaderboardsSteam::FindLeaderboard(const FString& LeaderboardName)
 {
 	FScopeLock ScopeLock(&LeaderboardMetadataLock);
 	FLeaderboardMetadataSteam* LeaderboardMetadata = GetLeaderboardMetadata(LeaderboardName);
diff --git a/Engine/Plugins/Online/OnlineSubsystemSteam/Source/Private/OnlineLeaderboardInterfaceSteam.h b/Engine/Plugins/Online/OnlineSubsystemSteam/Source/Private/OnlineLeaderboardInterfaceSteam.h
--- a/Engine/Plugins/Online/OnlineSubsystemSteam/Source/Private/OnlineLeaderboardInterfaceSteam.h
+++ b/Engine/Plugins/Online/OnlineSubsystemSteam/Source/Private/OnlineLeaderboardInterfaceSteam.h
@@ -72,7 +72,7 @@ PACKAGE_SCOPE:
 	 * @param LeaderboardName name of leaderboard to get information for
 	 * @return leaderboard metadata if it exists, else NULL
 	 */
-	FLeaderboardMetadataSteam* GetLeaderboardMetadata(const FName& LeaderboardName);
+	FLeaderboardMetadataSteam* GetLeaderboardMetadata(const FString& LeaderboardName);
 
 	/**
 	 *	Start an async task to create a leaderboard with the Steam backend
@@ -81,14 +81,14 @@ PACKAGE_SCOPE:
 	 * @param SortMethod method the leaderboard scores will be sorted, ignored if leaderboard exists
 	 * @param DisplayFormat type of data the leaderboard represents, ignored if leaderboard exists
 	 */
-	void CreateLeaderboard(const FName& LeaderboardName, ELeaderboardSort::Type SortMethod, ELeaderboardFormat::Type DisplayFormat);
+	void CreateLeaderboard(const FString& LeaderboardName, ELeaderboardSort::Type SortMethod, ELeaderboardFormat::Type DisplayFormat);
 
 	/**
 	 *	Start an async task to find a leaderboard with the Steam backend
 	 * If the leaderboard doesn't exist, a warning will be generated
 	 * @param LeaderboardName name of leaderboard to create
 	 */
-	void FindLeaderboard(const FName& LeaderboardName);
+	void FindLeaderboard(const FString& LeaderboardName);
 
 	/**
 	 *	Request the logged in user's stats from Steam
diff --git a/Engine/Plugins/Online/OnlineSubsystemSteam/Source/Private/OnlineSubsystemSteamTypes.h b/Engine/Plugins/Online/OnlineSubsystemSteam/Source/Private/OnlineSubsystemSteamTypes.h
--- a/Engine/Plugins/Online/OnlineSubsystemSteam/Source/Private/OnlineSubsystemSteamTypes.h
+++ b/Engine/Plugins/Online/OnlineSubsystemSteam/Source/Private/OnlineSubsystemSteamTypes.h
@@ -668,7 +668,7 @@ class FLeaderboardMetadataSteam
 public:
 
 	/** Name of leaderboard, matches Steam backend */
-    FName LeaderboardName;
+	FString LeaderboardName;
 	/** Sort Method */
 	ELeaderboardSort::Type SortMethod;
 	/** Display Type */
@@ -680,7 +680,7 @@ public:
 	/** State of the leaderboard handle download */
 	EOnlineAsyncTaskState::Type AsyncState;
 
-	FLeaderboardMetadataSteam(const FName& InLeaderboardName, ELeaderboardSort::Type InSortMethod, ELeaderboardFormat::Type InDisplayFormat) :
+	FLeaderboardMetadataSteam(const FString& InLeaderboardName, ELeaderboardSort::Type InSortMethod, ELeaderboardFormat::Type InDisplayFormat) :
 		LeaderboardName(InLeaderboardName),
 		SortMethod(InSortMethod),
 		DisplayFormat(InDisplayFormat),
@@ -690,7 +690,7 @@ public:
 	{
 	}
 
-	FLeaderboardMetadataSteam(const FName& InLeaderboardName) :
+	FLeaderboardMetadataSteam(const FString& InLeaderboardName) :
 		LeaderboardName(InLeaderboardName),
 		SortMethod(ELeaderboardSort::None),
 		DisplayFormat(ELeaderboardFormat::Number),
diff --git a/Engine/Plugins/Online/OnlineSubsystemUtils/Source/OnlineSubsystemUtils/Classes/AchievementWriteCallbackProxy.h b/Engine/Plugins/Online/OnlineSubsystemUtils/Source/OnlineSubsystemUtils/Classes/AchievementWriteCallbackProxy.h
--- a/Engine/Plugins/Online/OnlineSubsystemUtils/Source/OnlineSubsystemUtils/Classes/AchievementWriteCallbackProxy.h
+++ b/Engine/Plugins/Online/OnlineSubsystemUtils/Source/OnlineSubsystemUtils/Classes/AchievementWriteCallbackProxy.h
@@ -8,6 +8,7 @@
 class APlayerController;
 
 DECLARE_DYNAMIC_MULTICAST_DELEGATE_ThreeParams(FAchievementWriteDelegate, FName, WrittenAchievementName, float, WrittenProgress, int32, WrittenUserTag);
+DECLARE_DYNAMIC_MULTICAST_DELEGATE_ThreeParams(FAchievementWriteCompleteDelegate, FString, WrittenAchievementName, float, WrittenProgress, int32, WrittenUserTag);
 
 UCLASS(MinimalAPI)
 class UAchievementWriteCallbackProxy : public UOnlineBlueprintCallProxyBase
@@ -16,18 +17,30 @@ class UAchievementWriteCallbackProxy : public UOnlineBlueprintCallProxyBase
 
 	// Called when there is a successful achievement write
 	UPROPERTY(BlueprintAssignable)
-	FAchievementWriteDelegate OnSuccess;
+	FAchievementWriteCompleteDelegate OnWriteSuccess;
 
 	// Called when there is an unsuccessful achievement write
 	UPROPERTY(BlueprintAssignable)
+	FAchievementWriteCompleteDelegate OnWriteFailure;
+
+	UE_DEPRECATED(5.5, "Use OnWriteSuccess instead.")
+	UPROPERTY(BlueprintAssignable, meta = (DeprecatedProperty, DeprecationMessage = "Use OnWriteSuccess instead."))
+	FAchievementWriteDelegate OnSuccess;
+
+	UE_DEPRECATED(5.5, "Use OnWriteFailure instead.")
+	UPROPERTY(BlueprintAssignable, meta = (DeprecatedProperty, DeprecationMessage = "Use OnWriteFailure instead."))
 	FAchievementWriteDelegate OnFailure;
 
 	// Writes progress about an achievement to the default online subsystem
 	//   AchievementName is the ID of the achievement to update progress on
 	//   Progress is the reported progress toward accomplishing the achievement
 	//   UserTag is not used internally, but it is returned on success or failure
+	UE_DEPRECATED(5.5, "Use WriteProgress instead.")
+	UFUNCTION(BlueprintCallable, meta=(BlueprintInternalUseOnly = "true", WorldContext="WorldContextObject", DeprecatedFunction, DeprecationMessage = "Use WriteProgress instead."), Category="Online|Achievements")
+	static UAchievementWriteCallbackProxy* WriteAchievementProgress(UObject* WorldContextObject, APlayerController* PlayerController, FName AchievementName, float Progress = 100.0f, int32 UserTag = 0) { return WriteProgress(WorldContextObject, PlayerController, AchievementName.ToString(), Progress, UserTag); }
+
 	UFUNCTION(BlueprintCallable, meta=(BlueprintInternalUseOnly = "true", WorldContext="WorldContextObject"), Category="Online|Achievements")
-	static UAchievementWriteCallbackProxy* WriteAchievementProgress(UObject* WorldContextObject, APlayerController* PlayerController, FName AchievementName, float Progress = 100.0f, int32 UserTag = 0);
+	static UAchievementWriteCallbackProxy* WriteProgress(UObject* WorldContextObject, APlayerController* PlayerController, FString AchievementName, float Progress = 100.0f, int32 UserTag = 0);
 
 	// UOnlineBlueprintCallProxyBase interface
 	virtual void Activate() override;
@@ -49,7 +62,7 @@ private:
 	FOnlineAchievementsWritePtr WriteObject;
 
 	/** The achievement name */
-	FName AchievementName;
+	FString AchievementName;
 
 	/** The amount of progress made towards the achievement */
 	float AchievementProgress;
diff --git a/Engine/Plugins/Online/OnlineSubsystemUtils/Source/OnlineSubsystemUtils/Classes/LeaderboardQueryCallbackProxy.h b/Engine/Plugins/Online/OnlineSubsystemUtils/Source/OnlineSubsystemUtils/Classes/LeaderboardQueryCallbackProxy.h
--- a/Engine/Plugins/Online/OnlineSubsystemUtils/Source/OnlineSubsystemUtils/Classes/LeaderboardQueryCallbackProxy.h
+++ b/Engine/Plugins/Online/OnlineSubsystemUtils/Source/OnlineSubsystemUtils/Classes/LeaderboardQueryCallbackProxy.h
@@ -22,7 +22,7 @@ class ULeaderboardQueryCallbackProxy : public UObject
 
 	// Queries a leaderboard for an integer value
 	UFUNCTION(BlueprintCallable, meta = (BlueprintInternalUseOnly = "true", DisplayName="Read Leaderboard Integer"), Category="Online|Leaderboard")
-	static ULeaderboardQueryCallbackProxy* CreateProxyObjectForIntQuery(class APlayerController* PlayerController, FName StatName);
+	static ULeaderboardQueryCallbackProxy* CreateProxyObjectForIntQuery(class APlayerController* PlayerController, const FString& StatName);
 
 public:
 	//~ Begin UObject Interface
@@ -38,7 +38,7 @@ private:
 	void RemoveDelegate();
 
 	/** Triggers the query for a specifed user; the ReadObject must already be set up */
-	void TriggerQuery(class APlayerController* PlayerController, FName InStatName, EOnlineKeyValuePairDataType::Type StatType);
+	void TriggerQuery(class APlayerController* PlayerController, const FString& InStatName, EOnlineKeyValuePairDataType::Type StatType);
 
 private:
 	/** Delegate called when a leaderboard has been successfully read */
@@ -54,7 +54,7 @@ private:
 	bool bFailedToEvenSubmit;
 
 	/** Name of the stat being queried */
-	FName StatName;
+	FString StatName;
 
 	// Pointer to the world, needed to delay the results slightly
 	TWeakObjectPtr<UWorld> WorldPtr;
diff --git a/Engine/Plugins/Online/OnlineSubsystemUtils/Source/OnlineSubsystemUtils/Private/AchievementWriteCallbackProxy.cpp b/Engine/Plugins/Online/OnlineSubsystemUtils/Source/OnlineSubsystemUtils/Private/AchievementWriteCallbackProxy.cpp
--- a/Engine/Plugins/Online/OnlineSubsystemUtils/Source/OnlineSubsystemUtils/Private/AchievementWriteCallbackProxy.cpp
+++ b/Engine/Plugins/Online/OnlineSubsystemUtils/Source/OnlineSubsystemUtils/Private/AchievementWriteCallbackProxy.cpp
@@ -17,7 +17,7 @@ UAchievementWriteCallbackProxy::UAchievementWriteCallbackProxy(const FObjectInit
 {
 }
 
-UAchievementWriteCallbackProxy* UAchievementWriteCallbackProxy::WriteAchievementProgress(UObject* WorldContextObject, class APlayerController* PlayerController, FName AchievementName, float Progress, int32 UserTag)
+UAchievementWriteCallbackProxy* UAchievementWriteCallbackProxy::WriteProgress(UObject* WorldContextObject, class APlayerController* PlayerController, FString AchievementName, float Progress, int32 UserTag)
 {
 	UAchievementWriteCallbackProxy* Proxy = NewObject<UAchievementWriteCallbackProxy>();
 
@@ -56,7 +56,12 @@ void UAchievementWriteCallbackProxy::Activate()
 	}
 
 	// Fail immediately
-	OnFailure.Broadcast(AchievementName, AchievementProgress, UserTag);
+	OnWriteFailure.Broadcast(AchievementName, AchievementProgress, UserTag);
+	
+PRAGMA_DISABLE_DEPRECATION_WARNINGS
+	OnFailure.Broadcast(FName(AchievementName), AchievementProgress, UserTag);
+PRAGMA_ENABLE_DEPRECATION_WARNINGS
+
 	WriteObject.Reset();
 }
 
@@ -64,11 +69,19 @@ void UAchievementWriteCallbackProxy::OnAchievementWritten(const FUniqueNetId& Us
 {
 	if (bSuccess)
 	{
-		OnSuccess.Broadcast(AchievementName, AchievementProgress, UserTag);
+		OnWriteSuccess.Broadcast(AchievementName, AchievementProgress, UserTag);
+
+PRAGMA_DISABLE_DEPRECATION_WARNINGS
+		OnSuccess.Broadcast(FName(AchievementName), AchievementProgress, UserTag);	
+PRAGMA_ENABLE_DEPRECATION_WARNINGS
 	}
 	else
 	{
-		OnFailure.Broadcast(AchievementName, AchievementProgress, UserTag);
+		OnWriteFailure.Broadcast(AchievementName, AchievementProgress, UserTag);
+		
+PRAGMA_DISABLE_DEPRECATION_WARNINGS
+		OnFailure.Broadcast(FName(AchievementName), AchievementProgress, UserTag);	
+PRAGMA_ENABLE_DEPRECATION_WARNINGS
 	}
 
 	WriteObject.Reset();
diff --git a/Engine/Plugins/Online/OnlineSubsystemUtils/Source/OnlineSubsystemUtils/Private/LeaderboardBlueprintLibrary.cpp b/Engine/Plugins/Online/OnlineSubsystemUtils/Source/OnlineSubsystemUtils/Private/LeaderboardBlueprintLibrary.cpp
--- a/Engine/Plugins/Online/OnlineSubsystemUtils/Source/OnlineSubsystemUtils/Private/LeaderboardBlueprintLibrary.cpp
+++ b/Engine/Plugins/Online/OnlineSubsystemUtils/Source/OnlineSubsystemUtils/Private/LeaderboardBlueprintLibrary.cpp
@@ -62,12 +62,13 @@ bool ULeaderboardBlueprintLibrary::WriteLeaderboardObject(APlayerController* Pla
 bool ULeaderboardBlueprintLibrary::WriteLeaderboardInteger(APlayerController* PlayerController, FName StatName, int32 StatValue)
 {
 	FOnlineLeaderboardWrite WriteObject;
-	WriteObject.LeaderboardNames.Add(StatName);
-	WriteObject.RatedStat = StatName;
+	const FString StatNameString = StatName.ToString();
+	WriteObject.LeaderboardNames.Add(StatNameString);
+	WriteObject.RatedStat = StatNameString;
 	WriteObject.DisplayFormat = ELeaderboardFormat::Number;
 	WriteObject.SortMethod = ELeaderboardSort::Descending;
 	WriteObject.UpdateMethod = ELeaderboardUpdateMethod::KeepBest;
-	WriteObject.SetIntStat(StatName, StatValue);
+	WriteObject.SetIntStat(StatNameString, StatValue);
 
 	return WriteLeaderboardObject(PlayerController, WriteObject);
 }
diff --git a/Engine/Plugins/Online/OnlineSubsystemUtils/Source/OnlineSubsystemUtils/Private/LeaderboardQueryCallbackProxy.cpp b/Engine/Plugins/Online/OnlineSubsystemUtils/Source/OnlineSubsystemUtils/Private/LeaderboardQueryCallbackProxy.cpp
--- a/Engine/Plugins/Online/OnlineSubsystemUtils/Source/OnlineSubsystemUtils/Private/LeaderboardQueryCallbackProxy.cpp
+++ b/Engine/Plugins/Online/OnlineSubsystemUtils/Source/OnlineSubsystemUtils/Private/LeaderboardQueryCallbackProxy.cpp
@@ -17,7 +17,7 @@ ULeaderboardQueryCallbackProxy::ULeaderboardQueryCallbackProxy(const FObjectInit
 {
 }
 
-void ULeaderboardQueryCallbackProxy::TriggerQuery(APlayerController* PlayerController, FName InStatName, EOnlineKeyValuePairDataType::Type StatType)
+void ULeaderboardQueryCallbackProxy::TriggerQuery(APlayerController* PlayerController, const FString& InStatName, EOnlineKeyValuePairDataType::Type StatType)
 {
 	bFailedToEvenSubmit = true;
 
@@ -36,9 +36,9 @@ void ULeaderboardQueryCallbackProxy::TriggerQuery(APlayerController* PlayerContr
 
 					StatName = InStatName;
 					ReadObject = MakeShareable(new FOnlineLeaderboardRead());
-					ReadObject->LeaderboardName = StatName;
-					ReadObject->SortedColumn = StatName;
-					new (ReadObject->ColumnMetadata) FColumnMetaData(StatName, StatType);
+					ReadObject->LeaderboardName = InStatName;
+					ReadObject->SortedColumn = InStatName;
+					new (ReadObject->ColumnMetadata) FColumnMetaData(InStatName, StatType);
 
 					// Register the completion callback
 					LeaderboardReadCompleteDelegate       = FOnLeaderboardReadCompleteDelegate::CreateUObject(this, &ULeaderboardQueryCallbackProxy::OnStatsRead);
@@ -155,7 +155,7 @@ void ULeaderboardQueryCallbackProxy::BeginDestroy()
 	Super::BeginDestroy();
 }
 
-ULeaderboardQueryCallbackProxy* ULeaderboardQueryCallbackProxy::CreateProxyObjectForIntQuery(class APlayerController* PlayerController, FName StatName)
+ULeaderboardQueryCallbackProxy* ULeaderboardQueryCallbackProxy::CreateProxyObjectForIntQuery(class APlayerController* PlayerController, const FString& StatName)
 {
 	ULeaderboardQueryCallbackProxy* Proxy = NewObject<ULeaderboardQueryCallbackProxy>();
 	Proxy->SetFlags(RF_StrongRefOnFrame);
diff --git a/Engine/Plugins/Online/OnlineSubsystemUtils/Source/OnlineSubsystemUtils/Private/Tests/TestAchievementsInterface.cpp b/Engine/Plugins/Online/OnlineSubsystemUtils/Source/OnlineSubsystemUtils/Private/Tests/TestAchievementsInterface.cpp
--- a/Engine/Plugins/Online/OnlineSubsystemUtils/Source/OnlineSubsystemUtils/Private/Tests/TestAchievementsInterface.cpp
+++ b/Engine/Plugins/Online/OnlineSubsystemUtils/Source/OnlineSubsystemUtils/Private/Tests/TestAchievementsInterface.cpp
@@ -145,16 +145,16 @@ void FTestAchievementsInterface::WriteAchievements()
 	if (PlayerAchievements.Num() >= 3)
 	{
 		// report progress on first, last and median achievements
-		WriteObject->SetFloatStat(*PlayerAchievements[ 0 ].Id, 1.0f);
-		WriteObject->SetFloatStat(*PlayerAchievements[ PlayerAchievements.Num() - 1 ].Id, 50.0f);
-		WriteObject->SetFloatStat(*PlayerAchievements[ PlayerAchievements.Num() / 2 ].Id, 100.0f);
+		WriteObject->SetFloatStat(PlayerAchievements[ 0 ].Id, 1.0f);
+		WriteObject->SetFloatStat(PlayerAchievements[ PlayerAchievements.Num() - 1 ].Id, 50.0f);
+		WriteObject->SetFloatStat(PlayerAchievements[ PlayerAchievements.Num() / 2 ].Id, 100.0f);
 	}
 	else
 	{
 		// report progress on all achievements
 		for (int32 IdxAch = 0; IdxAch < PlayerAchievements.Num(); ++IdxAch)
 		{
-			WriteObject->SetFloatStat(*PlayerAchievements[ IdxAch ].Id, 50.0f);
+			WriteObject->SetFloatStat(PlayerAchievements[ IdxAch ].Id, 50.0f);
 		}
 	}
 
diff --git a/Engine/Plugins/Online/OnlineSubsystemUtils/Source/OnlineSubsystemUtils/Private/Tests/TestLeaderboardInterface.cpp b/Engine/Plugins/Online/OnlineSubsystemUtils/Source/OnlineSubsystemUtils/Private/Tests/TestLeaderboardInterface.cpp
--- a/Engine/Plugins/Online/OnlineSubsystemUtils/Source/OnlineSubsystemUtils/Private/Tests/TestLeaderboardInterface.cpp
+++ b/Engine/Plugins/Online/OnlineSubsystemUtils/Source/OnlineSubsystemUtils/Private/Tests/TestLeaderboardInterface.cpp
@@ -15,8 +15,8 @@ public:
 	TestLeaderboardWrite()
 	{
 		// Default properties
-		new (LeaderboardNames) FName(TEXT("TestLeaderboard"));
-		RatedStat = "TestIntStat1";
+		new (LeaderboardNames) FString("TestLeaderboard");
+		RatedStat = FString("TestIntStat1");
 		DisplayFormat = ELeaderboardFormat::Number;
 		SortMethod = ELeaderboardSort::Descending;
 		UpdateMethod = ELeaderboardUpdateMethod::KeepBest;
@@ -31,12 +31,12 @@ class TestLeaderboardRead : public FOnlineLeaderboardRead
 public:
 	TestLeaderboardRead(const FString& InLeaderboardName, const FString& InSortedColumn, const TMap<FString, EOnlineKeyValuePairDataType::Type>& InColumns)
 	{
-		LeaderboardName = FName(InLeaderboardName);
-		SortedColumn = FName(InSortedColumn);
+		LeaderboardName = InLeaderboardName;
+		SortedColumn = InSortedColumn;
 
 		for (TPair<FString, EOnlineKeyValuePairDataType::Type> Column : InColumns)
 		{
-			new (ColumnMetadata) FColumnMetaData(FName(Column.Key), Column.Value);
+			new (ColumnMetadata) FColumnMetaData(Column.Key, Column.Value);
 		}
 	}
 };
@@ -198,8 +198,8 @@ void FTestLeaderboardInterface::WriteLeaderboards()
 	TestLeaderboardWrite WriteObject;
 	
 	// Set some data
-	WriteObject.SetIntStat("TestIntStat1", 50);
-	WriteObject.SetFloatStat("TestFloatStat1", 99.5f);
+	WriteObject.SetIntStat(FString("TestIntStat1"), 50);
+	WriteObject.SetFloatStat(FString("TestFloatStat1"), 99.5f);
 
 	// Write it to the buffers
 	Leaderboards->WriteLeaderboards(TEXT("TEST"), *UserId, WriteObject);
@@ -230,7 +230,7 @@ void FTestLeaderboardInterface::PrintLeaderboards()
 
 		for (FStatsColumnArray::TConstIterator It(StatsRow.Columns); It; ++It)
 		{
-			UE_LOG_ONLINE_LEADERBOARD(Log, TEXT("     %s = %s"), *It.Key().ToString(), *It.Value().ToString());
+			UE_LOG_ONLINE_LEADERBOARD(Log, TEXT("     %s = %s"), *It.Key(), *It.Value().ToString());
 		}
 	}
 }
diff --git a/Engine/Source/Programs/Online/OSSTestsCore/Tests/Leaderboards/LeaderboardsFlushLeaderboardsTests.cpp b/Engine/Source/Programs/Online/OSSTestsCore/Tests/Leaderboards/LeaderboardsFlushLeaderboardsTests.cpp
--- a/Engine/Source/Programs/Online/OSSTestsCore/Tests/Leaderboards/LeaderboardsFlushLeaderboardsTests.cpp
+++ b/Engine/Source/Programs/Online/OSSTestsCore/Tests/Leaderboards/LeaderboardsFlushLeaderboardsTests.cpp
@@ -19,7 +19,7 @@ LEADERBOARDS_TEST_CASE("Verify calling FlushLeaderboards with valid inputs retur
 	FUniqueNetIdPtr LocalUserId = nullptr;
 
 	FName LocalSessionName = TEXT("FakeSessionName");
-	FName LocalNameForLeaderboard = TEXT("Name1");
+	const FString LocalNameForLeaderboard = TEXT("Name1");
 	FOnlineLeaderboardWrite WriteObject;
 
 	WriteObject.LeaderboardNames = { LocalNameForLeaderboard };
diff --git a/Engine/Source/Programs/Online/OSSTestsCore/Tests/Leaderboards/LeaderboardsReadLeaderboardsForFriendsTests.cpp b/Engine/Source/Programs/Online/OSSTestsCore/Tests/Leaderboards/LeaderboardsReadLeaderboardsForFriendsTests.cpp
--- a/Engine/Source/Programs/Online/OSSTestsCore/Tests/Leaderboards/LeaderboardsReadLeaderboardsForFriendsTests.cpp
+++ b/Engine/Source/Programs/Online/OSSTestsCore/Tests/Leaderboards/LeaderboardsReadLeaderboardsForFriendsTests.cpp
@@ -23,7 +23,7 @@ EG_LEADERBOARDS_TEST_CASE("Verify calling ReadLeaderboardsForFriends with valid
 	FUniqueNetIdPtr LocalUserId = nullptr;
 
 	FName LocalSessionName = TEXT("FakeSessionName");
-	FName LocalNameForLeaderboard = TEXT("Name1");
+	const FString LocalNameForLeaderboard = TEXT("Name1");
 	FOnlineLeaderboardWrite WriteObject;
 
 	WriteObject.LeaderboardNames = { LocalNameForLeaderboard };
diff --git a/Engine/Source/Programs/Online/OSSTestsCore/Tests/Leaderboards/LeaderboardsReadLeaderboardsTests.cpp b/Engine/Source/Programs/Online/OSSTestsCore/Tests/Leaderboards/LeaderboardsReadLeaderboardsTests.cpp
--- a/Engine/Source/Programs/Online/OSSTestsCore/Tests/Leaderboards/LeaderboardsReadLeaderboardsTests.cpp
+++ b/Engine/Source/Programs/Online/OSSTestsCore/Tests/Leaderboards/LeaderboardsReadLeaderboardsTests.cpp
@@ -19,7 +19,7 @@ EG_LEADERBOARDS_TEST_CASE("Verify calling ReadLeaderboards with valid inputs ret
 	FUniqueNetIdPtr LocalUserId = nullptr;
 
 	FName LocalSessionName = TEXT("FakeSessionName");
-	FName LocalNameForLeaderboard = TEXT("Name1");
+	const FString LocalNameForLeaderboard = TEXT("Name1");
 	FOnlineLeaderboardWrite WriteObject;
 
 	WriteObject.LeaderboardNames = { LocalNameForLeaderboard };
diff --git a/Engine/Source/Programs/Online/OSSTestsCore/Tests/Leaderboards/LeaderboardsWriteLeaderboardsTests.cpp b/Engine/Source/Programs/Online/OSSTestsCore/Tests/Leaderboards/LeaderboardsWriteLeaderboardsTests.cpp
--- a/Engine/Source/Programs/Online/OSSTestsCore/Tests/Leaderboards/LeaderboardsWriteLeaderboardsTests.cpp
+++ b/Engine/Source/Programs/Online/OSSTestsCore/Tests/Leaderboards/LeaderboardsWriteLeaderboardsTests.cpp
@@ -18,7 +18,7 @@ LEADERBOARDS_TEST_CASE("SubsystemNull. Verify calling WriteLeaderboards with val
 	FUniqueNetIdPtr LocalUserId = nullptr;
 
 	FName LocalSessionName = TEXT("FakeSessionName");
-	FName LocalNameForLeaderboard = TEXT("Name1");
+	const FString LocalNameForLeaderboard = TEXT("Name1");
 	FOnlineLeaderboardWrite WriteObject;
 
 	WriteObject.LeaderboardNames = { LocalNameForLeaderboard };
