Addition of missing callsite updates after deprecation.
Fix for Horde issue 536919

diff --git a/Engine/Plugins/Online/OnlineSubsystem/Source/Public/OnlineStats.h b/Engine/Plugins/Online/OnlineSubsystem/Source/Public/OnlineStats.h
--- a/Engine/Plugins/Online/OnlineSubsystem/Source/Public/OnlineStats.h
+++ b/Engine/Plugins/Online/OnlineSubsystem/Source/Public/OnlineStats.h
@@ -129,6 +129,11 @@ public:
 	FORCEINLINE FNameArrayDeprecationWrapper& operator=(FNameArrayDeprecationWrapper&& Other) { Super::operator=(MoveTemp(Other)); return *this; }
 	FORCEINLINE FNameArrayDeprecationWrapper& operator=(const FNameArrayDeprecationWrapper& Other) { Super::operator=(Other); return *this; }
 
+	FORCEINLINE FNameArrayDeprecationWrapper(TArray<FString>&& Other) : Super(MoveTemp(Other)) {}
+	FORCEINLINE FNameArrayDeprecationWrapper(const TArray<FString>& Other) : Super(Other) {}
+	FORCEINLINE FNameArrayDeprecationWrapper& operator=(TArray<FString>&& Other) { Super::operator=(MoveTemp(Other)); return *this; }
+	FORCEINLINE FNameArrayDeprecationWrapper& operator=(const TArray<FString>& Other) { Super::operator=(Other); return *this; }
+
 	using Super::Add;
 	UE_DEPRECATED(5.5, "This variable is now a TArray<FString> instead of a TArray<FName>. Please update your code accordingly")
 		int32 Add(const FName& InElement)
diff --git a/Engine/Source/Programs/Online/OSSTestsCore/Tests/Leaderboards/LeaderboardsFlushLeaderboardsTests.cpp b/Engine/Source/Programs/Online/OSSTestsCore/Tests/Leaderboards/LeaderboardsFlushLeaderboardsTests.cpp
--- a/Engine/Source/Programs/Online/OSSTestsCore/Tests/Leaderboards/LeaderboardsFlushLeaderboardsTests.cpp
+++ b/Engine/Source/Programs/Online/OSSTestsCore/Tests/Leaderboards/LeaderboardsFlushLeaderboardsTests.cpp
@@ -25,7 +25,7 @@ LEADERBOARDS_TEST_CASE("Verify calling FlushLeaderboards with valid inputs retur
 	WriteObject.LeaderboardNames = { LocalNameForLeaderboard };
 	WriteObject.SortMethod = ELeaderboardSort::Ascending;
 	WriteObject.UpdateMethod = ELeaderboardUpdateMethod::KeepBest;
-	WriteObject.RatedStat = TEXT("Scores");
+	WriteObject.RatedStat = FString(TEXT("Scores"));
 	int32 NumUsersToImplicitLogin = 1;
 
 	GetLoginPipeline(NumUsersToImplicitLogin)
diff --git a/Engine/Source/Programs/Online/OSSTestsCore/Tests/Leaderboards/LeaderboardsReadLeaderboardsForFriendsTests.cpp b/Engine/Source/Programs/Online/OSSTestsCore/Tests/Leaderboards/LeaderboardsReadLeaderboardsForFriendsTests.cpp
--- a/Engine/Source/Programs/Online/OSSTestsCore/Tests/Leaderboards/LeaderboardsReadLeaderboardsForFriendsTests.cpp
+++ b/Engine/Source/Programs/Online/OSSTestsCore/Tests/Leaderboards/LeaderboardsReadLeaderboardsForFriendsTests.cpp
@@ -29,7 +29,7 @@ EG_LEADERBOARDS_TEST_CASE("Verify calling ReadLeaderboardsForFriends with valid
 	WriteObject.LeaderboardNames = { LocalNameForLeaderboard };
 	WriteObject.SortMethod = ELeaderboardSort::Ascending;
 	WriteObject.UpdateMethod = ELeaderboardUpdateMethod::KeepBest;
-	WriteObject.RatedStat = TEXT("Scores");
+	WriteObject.RatedStat = FString(TEXT("Scores"));
 
 	TArray<FUniqueNetIdRef> LocalPlayers;
 	FOnlineLeaderboardReadRef LocalReadObject = MakeShared<FOnlineLeaderboardRead>();
diff --git a/Engine/Source/Programs/Online/OSSTestsCore/Tests/Leaderboards/LeaderboardsReadLeaderboardsTests.cpp b/Engine/Source/Programs/Online/OSSTestsCore/Tests/Leaderboards/LeaderboardsReadLeaderboardsTests.cpp
--- a/Engine/Source/Programs/Online/OSSTestsCore/Tests/Leaderboards/LeaderboardsReadLeaderboardsTests.cpp
+++ b/Engine/Source/Programs/Online/OSSTestsCore/Tests/Leaderboards/LeaderboardsReadLeaderboardsTests.cpp
@@ -25,7 +25,7 @@ EG_LEADERBOARDS_TEST_CASE("Verify calling ReadLeaderboards with valid inputs ret
 	WriteObject.LeaderboardNames = { LocalNameForLeaderboard };
 	WriteObject.SortMethod = ELeaderboardSort::Ascending;
 	WriteObject.UpdateMethod = ELeaderboardUpdateMethod::KeepBest;
-	WriteObject.RatedStat = TEXT("Scores");
+	WriteObject.RatedStat = FString(TEXT("Scores"));
 
 	TArray<FUniqueNetIdRef> LocalPlayers;
 	FOnlineLeaderboardReadRef LocalReadObject = MakeShared<FOnlineLeaderboardRead>();
diff --git a/Engine/Source/Programs/Online/OSSTestsCore/Tests/Leaderboards/LeaderboardsWriteLeaderboardsTests.cpp b/Engine/Source/Programs/Online/OSSTestsCore/Tests/Leaderboards/LeaderboardsWriteLeaderboardsTests.cpp
--- a/Engine/Source/Programs/Online/OSSTestsCore/Tests/Leaderboards/LeaderboardsWriteLeaderboardsTests.cpp
+++ b/Engine/Source/Programs/Online/OSSTestsCore/Tests/Leaderboards/LeaderboardsWriteLeaderboardsTests.cpp
@@ -24,7 +24,7 @@ LEADERBOARDS_TEST_CASE("SubsystemNull. Verify calling WriteLeaderboards with val
 	WriteObject.LeaderboardNames = { LocalNameForLeaderboard };
 	WriteObject.SortMethod = ELeaderboardSort::Ascending;
 	WriteObject.UpdateMethod = ELeaderboardUpdateMethod::KeepBest;
-	WriteObject.RatedStat = TEXT("Scores");
+	WriteObject.RatedStat = FString(TEXT("Scores"));
 	int32 NumUsersToImplicitLogin = 1;
 
 	GetLoginPipeline(NumUsersToImplicitLogin)
