[AutoRTFM] Use FTransactionallySafeCriticalSection for FOnlineSubsystemModule::OnlineSubsystemsLock

This is locked from a closed transaction.

diff --git a/Engine/Plugins/Online/OnlineSubsystem/Source/Private/OnlineSubsystemModule.cpp b/Engine/Plugins/Online/OnlineSubsystem/Source/Private/OnlineSubsystemModule.cpp
--- a/Engine/Plugins/Online/OnlineSubsystem/Source/Private/OnlineSubsystemModule.cpp
+++ b/Engine/Plugins/Online/OnlineSubsystem/Source/Private/OnlineSubsystemModule.cpp
@@ -2,6 +2,7 @@
 
 #include "OnlineSubsystemModule.h"
 #include "Misc/ConfigCacheIni.h"
+#include "Misc/TransactionallySafeScopeLock.h"
 #include "OnlineSubsystemImpl.h"
 #include "OnlineDelegates.h"
 
@@ -227,7 +228,7 @@ void FOnlineSubsystemModule::ReloadDefaultSubsystem()
 void FOnlineSubsystemModule::PreUnloadOnlineSubsystem()
 {
 	// Shutdown all online subsystem instances
-	FScopeLock Lock(&OnlineSubsystemsLock);
+	FTransactionallySafeScopeLock Lock(&OnlineSubsystemsLock);
 	for (TMap<FName, IOnlineSubsystemPtr>::TIterator It(OnlineSubsystems); It; ++It)
 	{
 		It.Value()->PreUnload();
@@ -240,7 +241,7 @@ void FOnlineSubsystemModule::ShutdownOnlineSubsystem()
 
 	// Shutdown all online subsystem instances
 	{
-		FScopeLock Lock(&OnlineSubsystemsLock);
+		FTransactionallySafeScopeLock Lock(&OnlineSubsystemsLock);
 		for (TMap<FName, IOnlineSubsystemPtr>::TIterator It(OnlineSubsystems); It; ++It)
 		{
 			It.Value()->Shutdown();
@@ -280,7 +281,7 @@ void FOnlineSubsystemModule::UnregisterPlatformService(const FName FactoryName)
 
 void FOnlineSubsystemModule::EnumerateOnlineSubsystems(FEnumerateOnlineSubsystemCb& EnumCb)
 {
-	FScopeLock Lock(&OnlineSubsystemsLock);
+	FTransactionallySafeScopeLock Lock(&OnlineSubsystemsLock);
 	for (TPair<FName, IOnlineSubsystemPtr>& OnlineSubsystem : OnlineSubsystems)
 	{
 		if (OnlineSubsystem.Value.IsValid())
@@ -373,7 +374,7 @@ IOnlineSubsystem* FOnlineSubsystemModule::GetOnlineSubsystem(const FName InSubsy
 	{
 		bool bWasNewlyCreated = false;
 		{
-			FScopeLock Lock(&OnlineSubsystemsLock);
+			FTransactionallySafeScopeLock Lock(&OnlineSubsystemsLock);
 			OnlineSubsystem = FindExistingSubsystem(OnlineSubsystems, KeyName);
 			if (!OnlineSubsystem)
 			{
@@ -448,7 +449,7 @@ void FOnlineSubsystemModule::DestroyOnlineSubsystem(const FName InSubsystemName)
 	{
 		IOnlineSubsystemPtr OnlineSubsystem;
 		{
-			FScopeLock Lock(&OnlineSubsystemsLock);
+			FTransactionallySafeScopeLock Lock(&OnlineSubsystemsLock);
 			OnlineSubsystems.RemoveAndCopyValue(KeyName, OnlineSubsystem);
 		}
 		if (OnlineSubsystem.IsValid())
@@ -473,7 +474,7 @@ bool FOnlineSubsystemModule::DoesInstanceExist(const FName InSubsystemName) cons
 	{
 		IOnlineSubsystemPtr OnlineSubsystem;
 		{
-			FScopeLock Lock(&OnlineSubsystemsLock);
+			FTransactionallySafeScopeLock Lock(&OnlineSubsystemsLock);
 			OnlineSubsystem = FindExistingSubsystem(OnlineSubsystems, KeyName);
 		}
 		return OnlineSubsystem.IsValid();
diff --git a/Engine/Plugins/Online/OnlineSubsystem/Source/Public/OnlineSubsystemModule.h b/Engine/Plugins/Online/OnlineSubsystem/Source/Public/OnlineSubsystemModule.h
--- a/Engine/Plugins/Online/OnlineSubsystem/Source/Public/OnlineSubsystemModule.h
+++ b/Engine/Plugins/Online/OnlineSubsystem/Source/Public/OnlineSubsystemModule.h
@@ -2,6 +2,7 @@
 
 #pragma once
 
+#include "Misc/TransactionallySafeCriticalSection.h"
 #include "Modules/ModuleInterface.h"
 #include "Templates/SharedPointer.h"
 
@@ -51,7 +52,7 @@ private:
 	 * The intent is OnlineSubsystems can have elements added on separate threads through normal gameplay, though this is expected to be very rare and only once per subsystem type.
 	 * OnlineSubsystems is not expected to remove elements except in very controlled scenarios.  In that case, the lock is only limiting access to the OnlineSubsystems array, but the result may be invalidated.
 	 */
-	mutable FCriticalSection OnlineSubsystemsLock;
+	mutable FTransactionallySafeCriticalSection OnlineSubsystemsLock;
 
 	/** Have we warned already for a given online subsystem creation failure */
 	TSet<FName> OnlineSubsystemFailureNotes;
