Removed redundant text in deprecation message.

diff --git a/Engine/Plugins/Online/OnlineSubsystem/Source/Public/OnlineStats.h b/Engine/Plugins/Online/OnlineSubsystem/Source/Public/OnlineStats.h
--- a/Engine/Plugins/Online/OnlineSubsystem/Source/Public/OnlineStats.h
+++ b/Engine/Plugins/Online/OnlineSubsystem/Source/Public/OnlineStats.h
@@ -135,14 +135,14 @@ public:
 	FORCEINLINE FNameArrayDeprecationWrapper& operator=(const TArray<FString>& Other) { Super::operator=(Other); return *this; }
 
 	using Super::Add;
-	UE_DEPRECATED(5.5, "This variable is now a TArray<FString> instead of a TArray<FName>. Please update your code accordingly")
+	UE_DEPRECATED(5.5, "This variable is now a TArray<FString> instead of a TArray<FName>.")
 		int32 Add(const FName& InElement)
 	{
 		return Super::Add(InElement.ToString());
 	}
 
 	using Super::Emplace;
-	UE_DEPRECATED(5.5, "This variable is now a TArray<FString> instead of a TArray<FName>. Please update your code accordingly")
+	UE_DEPRECATED(5.5, "This variable is now a TArray<FString> instead of a TArray<FName>.")
 		int32 Emplace(const FName& InElement)
 	{
 		return Super::Emplace(InElement.ToString());
@@ -165,22 +165,22 @@ public:
 	FORCEINLINE FNameDeprecationWrapper& operator=(FString&& Other) { Super::operator=(MoveTemp(Other)); return *this; }
 	FORCEINLINE FNameDeprecationWrapper& operator=(const FString& Other) { Super::operator=(Other); return *this; }
 
-	UE_DEPRECATED(5.5, "This variable is now an FString instead of an FName. Please update your code accordingly")
+	UE_DEPRECATED(5.5, "This variable is now an FString instead of an FName.")
 	FORCEINLINE FNameDeprecationWrapper(FName&& DeprecatedValue) : Super(DeprecatedValue.ToString()) {}
 
-	UE_DEPRECATED(5.5, "This variable is now an FString instead of an FName. Please update your code accordingly")
+	UE_DEPRECATED(5.5, "This variable is now an FString instead of an FName.")
 		FORCEINLINE FNameDeprecationWrapper(const FName& DeprecatedValue) : Super(DeprecatedValue.ToString()) {}
 
-	UE_DEPRECATED(5.5, "This variable is now an FString instead of an FName. Please update your code accordingly")
+	UE_DEPRECATED(5.5, "This variable is now an FString instead of an FName.")
 		FORCEINLINE FNameDeprecationWrapper& operator=(FName&& DeprecatedValue) { Super::operator=(DeprecatedValue.ToString()); return *this; }
 
-	UE_DEPRECATED(5.5, "This variable is now an FString instead of an FName. Please update your code accordingly")
+	UE_DEPRECATED(5.5, "This variable is now an FString instead of an FName.")
 		FORCEINLINE FNameDeprecationWrapper& operator=(const FName& DeprecatedValue) { Super::operator=(DeprecatedValue.ToString()); return *this; }
 
-	UE_DEPRECATED(5.5, "This variable is now an FString instead of an FName. Please update your code accordingly")
+	UE_DEPRECATED(5.5, "This variable is now an FString instead of an FName.")
 		FORCEINLINE FString ToString() const { return *this; }
 
-	UE_DEPRECATED(5.5, "This variable is now an FString instead of an FName. Please update your code accordingly")
+	UE_DEPRECATED(5.5, "This variable is now an FString instead of an FName.")
 		FORCEINLINE void ToString(FString& Out) const { Out = *this; }
 };
 
