Removed Online logic marked for deprecation in 5.3

diff --git a/Engine/Plugins/Experimental/Web/WebAPI/Source/WebAPI/Public/WebAPIHttpRequest.h b/Engine/Plugins/Experimental/Web/WebAPI/Source/WebAPI/Public/WebAPIHttpRequest.h
--- a/Engine/Plugins/Experimental/Web/WebAPI/Source/WebAPI/Public/WebAPIHttpRequest.h
+++ b/Engine/Plugins/Experimental/Web/WebAPI/Source/WebAPI/Public/WebAPIHttpRequest.h
@@ -215,15 +215,6 @@ TRequest<PayloadType>& TRequest<PayloadType>::BindCompletionCallback(TFunction<v
 	return *this;
 }
 
-PRAGMA_DISABLE_DEPRECATION_WARNINGS
-template <typename PayloadType>
-TRequest<PayloadType>& TRequest<PayloadType>::BindProgressCallback(TFunction<void(FHttpRequestPtr, int32, int32)> ProgressCallback)
-{
-	InternalRequest->OnRequestProgress().BindLambda(MoveTemp(ProgressCallback));
-	return *this;
-}
-PRAGMA_ENABLE_DEPRECATION_WARNINGS
-
 template <typename PayloadType>
 TRequest<PayloadType>& TRequest<PayloadType>::BindProgressCallback64(TFunction<void(FHttpRequestPtr, uint64, uint64)> ProgressCallback)
 {
diff --git a/Engine/Plugins/Online/EOSShared/Source/EOSShared/Private/EOSSDKManager.h b/Engine/Plugins/Online/EOSShared/Source/EOSShared/Private/EOSSDKManager.h
--- a/Engine/Plugins/Online/EOSShared/Source/EOSShared/Private/EOSSDKManager.h
+++ b/Engine/Plugins/Online/EOSShared/Source/EOSShared/Private/EOSSDKManager.h
@@ -37,9 +37,6 @@ public:
 	virtual ~FEOSSDKManager();
 
 	// Begin IEOSSDKManager
-	PRAGMA_DISABLE_DEPRECATION_WARNINGS
-	virtual EOS_EResult Initialize() override;
-	PRAGMA_ENABLE_DEPRECATION_WARNINGS
 	virtual bool IsInitialized() const override { return bInitialized; }
 
 	virtual const FEOSSDKPlatformConfig* GetPlatformConfig(const FString& PlatformConfigName, bool bLoadIfMissing = false) override;
@@ -68,6 +65,7 @@ public:
 	virtual void AddCallbackObject(TUniquePtr<class FCallbackBase> CallbackObj) override;
 	// End IEOSSDKManager
 
+	EOS_EResult Initialize();
 	void Shutdown();
 
 protected:
diff --git a/Engine/Plugins/Online/EOSShared/Source/EOSShared/Public/IEOSSDKManager.h b/Engine/Plugins/Online/EOSShared/Source/EOSShared/Public/IEOSSDKManager.h
--- a/Engine/Plugins/Online/EOSShared/Source/EOSShared/Public/IEOSSDKManager.h
+++ b/Engine/Plugins/Online/EOSShared/Source/EOSShared/Public/IEOSSDKManager.h
@@ -99,8 +99,6 @@ public:
 
 	virtual ~IEOSSDKManager() = default;
 
-	UE_DEPRECATED(5.3, "Initialize is now called automatically when the EOSShared module is started. User code should instead check IsInitialized to see whether the EOSSDK was successfully started.")
-	virtual EOS_EResult Initialize() = 0;
 	virtual bool IsInitialized() const = 0;
 
 	virtual const FEOSSDKPlatformConfig* GetPlatformConfig(const FString& PlatformConfigName, bool bLoadIfMissing = false) = 0;
diff --git a/Engine/Plugins/Online/OnlineFramework/Source/Hotfix/Public/OnlineHotfixManager.h b/Engine/Plugins/Online/OnlineFramework/Source/Hotfix/Public/OnlineHotfixManager.h
--- a/Engine/Plugins/Online/OnlineFramework/Source/Hotfix/Public/OnlineHotfixManager.h
+++ b/Engine/Plugins/Online/OnlineFramework/Source/Hotfix/Public/OnlineHotfixManager.h
@@ -283,15 +283,6 @@ protected:
 	 * @return whether the file was successfully processed
 	 */
 	virtual bool ApplyHotfixProcessing(const FCloudFileHeader& FileHeader);
-	/**
-	 * Called prior to reading the file data.
-	 *
-	 * @param FileData - byte data of the hotfix file
-	 *
-	 * @return whether the file was successfully preprocessed
-	 */
-	UE_DEPRECATED(5.3, "Replaced with PreProcessDownloadedFileData taking a FCloudFileHeader")
-	virtual bool PreProcessDownloadedFileData(TArray<uint8>& FileData) const { return true;	}
 	/**
 	 * Called prior to reading the file data.
 	 *
diff --git a/Engine/Plugins/Online/OnlineFramework/Source/Party/Private/Party/SocialParty.cpp b/Engine/Plugins/Online/OnlineFramework/Source/Party/Private/Party/SocialParty.cpp
--- a/Engine/Plugins/Online/OnlineFramework/Source/Party/Private/Party/SocialParty.cpp
+++ b/Engine/Plugins/Online/OnlineFramework/Source/Party/Private/Party/SocialParty.cpp
@@ -251,12 +251,6 @@ bool USocialParty::CanInviteUser(const USocialUser& User, ESocialPartyInviteMeth
 	return CanInviteUserInternal(User, InviteMethod) == ESocialPartyInviteFailureReason::Success;
 }
 
-ESocialPartyInviteFailureReason USocialParty::CanInviteUserInternal(const USocialUser& User) const
-{
-	UE_LOG(LogParty, Warning, TEXT("Invoking a deprecated method CanInviteUserInternal, use CanInviteUserInternal(const USocialUser& User, const ESocialPartyInviteMethod InviteMethod) instead!"));
-	return CanInviteUserInternal(User, ESocialPartyInviteMethod::Other);
-}
-
 ESocialPartyInviteFailureReason USocialParty::CanInviteUserInternal(const USocialUser& User, ESocialPartyInviteMethod InviteMethod) const
 {
 	bool bCanOnlyInviteFriend = false;
diff --git a/Engine/Plugins/Online/OnlineFramework/Source/Party/Public/Party/SocialParty.h b/Engine/Plugins/Online/OnlineFramework/Source/Party/Public/Party/SocialParty.h
--- a/Engine/Plugins/Online/OnlineFramework/Source/Party/Public/Party/SocialParty.h
+++ b/Engine/Plugins/Online/OnlineFramework/Source/Party/Public/Party/SocialParty.h
@@ -271,8 +271,6 @@ protected:
 	virtual void OnLeftPartyInternal(EMemberExitedReason Reason);
 
 	/** Virtual versions of the package-scoped "CanX" methods above, as a virtual declared within package scoping cannot link (exported public, imported protected) */
-	UE_DEPRECATED(5.3, "This function has been deperecated, use CanInviteUserInternal(const USocialUser& User, const ESocialPartyInviteMethod InviteMethod) instead.")
-	virtual ESocialPartyInviteFailureReason CanInviteUserInternal(const USocialUser& User) const;
 	virtual ESocialPartyInviteFailureReason CanInviteUserInternal(const USocialUser& User, const ESocialPartyInviteMethod InviteMethod) const;
 
 	UE_DEPRECATED(5.5, "This has been deprecated to support multiple local players and now requires a performing player.")
diff --git a/Engine/Plugins/Online/OnlineSubsystem/Source/Public/Interfaces/OnlineUserCloudInterface.h b/Engine/Plugins/Online/OnlineSubsystem/Source/Public/Interfaces/OnlineUserCloudInterface.h
--- a/Engine/Plugins/Online/OnlineSubsystem/Source/Public/Interfaces/OnlineUserCloudInterface.h
+++ b/Engine/Plugins/Online/OnlineSubsystem/Source/Public/Interfaces/OnlineUserCloudInterface.h
@@ -33,20 +33,6 @@ struct FAnalyticsEventAttribute;
 DECLARE_MULTICAST_DELEGATE_TwoParams(FOnEnumerateUserFilesComplete, bool, const FUniqueNetId&);
 typedef FOnEnumerateUserFilesComplete::FDelegate FOnEnumerateUserFilesCompleteDelegate;
 
-/**
- * Delegate fired at intervals during a user file write to the network platform's storage
- *
- * @param BytesWritten the number of bytes written so far
- * @param UserId User owning the storage
- * @param FileName the name of the file this was for
- *
- */
-UE_DEPRECATED(5.3, "OnWriteUserFileProgress has been deprecated, use OnWriteUserFileProgress64 instead")
-DECLARE_MULTICAST_DELEGATE_ThreeParams(FOnWriteUserFileProgress, int32, const FUniqueNetId&, const FString&);
-PRAGMA_DISABLE_DEPRECATION_WARNINGS
-typedef FOnWriteUserFileProgress::FDelegate FOnWriteUserFileProgressDelegate;
-PRAGMA_ENABLE_DEPRECATION_WARNINGS
-
 /**
  * Delegate fired at intervals during a user file write to the network platform's storage
  *
@@ -222,44 +208,6 @@ public:
 	 */
 	virtual bool WriteUserFile(const FUniqueNetId& UserId, const FString& FileName, TArray<uint8>& FileContents, bool bCompressBeforeUpload = false) = 0;
 
-	/**
-	* Delegate fired at intervals during a user file write to the network platform's storage
-	*
-	* @param BytesWritten the number of bytes written so far
-	* @param UserId User owning the storage
-	* @param FileName the name of the file this was for
-	*
-	*/
-PRAGMA_DISABLE_DEPRECATION_WARNINGS
-	// 	Expanded macro with deprecate warnings: DEFINE_ONLINE_DELEGATE_THREE_PARAM(OnWriteUserFileProgress, int32, const FUniqueNetId&, const FString&);
-
-	UE_DEPRECATED(5.3, "OnWriteUserFileProgress has been deprecated, use OnWriteUserFileProgress64 instead")
-	FOnWriteUserFileProgress OnWriteUserFileProgressDelegates;
-
-	UE_DEPRECATED(5.3, "AddOnWriteUserFileProgressDelegate_Handle has been deprecated, use AddOnWriteUserFileProgress64Delegate_Handle instead")
-	virtual FDelegateHandle AddOnWriteUserFileProgressDelegate_Handle(const FOnWriteUserFileProgressDelegate& Delegate)
-	{
-		OnWriteUserFileProgressDelegates.Add(Delegate);
-		return Delegate.GetHandle();
-	}
-	UE_DEPRECATED(5.3, "ClearOnWriteUserFileProgressDelegate_Handle has been deprecated, use ClearOnWriteUserFileProgress64Delegate_Handle instead")
-	virtual void ClearOnWriteUserFileProgressDelegate_Handle(FDelegateHandle& Handle)
-	{
-		OnWriteUserFileProgressDelegates.Remove(Handle);
-		Handle.Reset();
-	}
-	UE_DEPRECATED(5.3, "ClearOnWriteUserFileProgressDelegates has been deprecated, use ClearOnWriteUserFileProgress64Delegates instead")
-	virtual void ClearOnWriteUserFileProgressDelegates(void* Object)
-	{
-		OnWriteUserFileProgressDelegates.RemoveAll(Object);
-	}
-	UE_DEPRECATED(5.3, "TriggerOnWriteUserFileProgressDelegates has been deprecated, use TriggerOnWriteUserFileProgress64Delegates instead")
-	virtual void TriggerOnWriteUserFileProgressDelegates(int32 Param1, const FUniqueNetId& Param2, const FString& Param3)
-	{
-		OnWriteUserFileProgressDelegates.Broadcast(Param1, Param2, Param3);
-	}
-PRAGMA_ENABLE_DEPRECATION_WARNINGS
-
 	/**
 	* Delegate fired at intervals during a user file write to the network platform's storage
 	*
diff --git a/Engine/Plugins/Online/OnlineSubsystemEOS/Source/OnlineSubsystemEOS/Private/OnlineUserCloudEOS.cpp b/Engine/Plugins/Online/OnlineSubsystemEOS/Source/OnlineSubsystemEOS/Private/OnlineUserCloudEOS.cpp
--- a/Engine/Plugins/Online/OnlineSubsystemEOS/Source/OnlineSubsystemEOS/Private/OnlineUserCloudEOS.cpp
+++ b/Engine/Plugins/Online/OnlineSubsystemEOS/Source/OnlineSubsystemEOS/Private/OnlineUserCloudEOS.cpp
@@ -553,9 +553,7 @@ bool FOnlineUserCloudEOS::WriteUserFile(const FUniqueNetId& UserId, const FStrin
 	CallbackObj->SetNested2CallbackLambda([this, SharedUserId, FileName](const EOS_PlayerDataStorage_FileTransferProgressCallbackInfo* Data)
 		{
 			UE_LOG_ONLINE_CLOUD(VeryVerbose, TEXT("[FOnlineUserCloudEOS::WriteUserFile] File transfer progress for file %s is %d bytes"), *FileName, Data->BytesTransferred);
-PRAGMA_DISABLE_DEPRECATION_WARNINGS
-			TriggerOnWriteUserFileProgressDelegates(Data->BytesTransferred, *SharedUserId, FileName);
-PRAGMA_ENABLE_DEPRECATION_WARNINGS
+
 			TriggerOnWriteUserFileProgress64Delegates(Data->BytesTransferred, *SharedUserId, FileName);
 		});
 
diff --git a/Engine/Plugins/Online/OnlineSubsystemEOS/Source/OnlineSubsystemEOSPlus/Private/OnlineUserCloudEOSPlus.cpp b/Engine/Plugins/Online/OnlineSubsystemEOS/Source/OnlineSubsystemEOSPlus/Private/OnlineUserCloudEOSPlus.cpp
--- a/Engine/Plugins/Online/OnlineSubsystemEOS/Source/OnlineSubsystemEOSPlus/Private/OnlineUserCloudEOSPlus.cpp
+++ b/Engine/Plugins/Online/OnlineSubsystemEOS/Source/OnlineSubsystemEOSPlus/Private/OnlineUserCloudEOSPlus.cpp
@@ -15,9 +15,6 @@ FOnlineUserCloudEOSPlus::~FOnlineUserCloudEOSPlus()
 	if (BaseUserCloudInterface.IsValid())
 	{
 		BaseUserCloudInterface->ClearOnEnumerateUserFilesCompleteDelegates(this);
-PRAGMA_DISABLE_DEPRECATION_WARNINGS
-		BaseUserCloudInterface->ClearOnWriteUserFileProgressDelegates(this);
-PRAGMA_ENABLE_DEPRECATION_WARNINGS
 		BaseUserCloudInterface->ClearOnWriteUserFileProgress64Delegates(this);
 		BaseUserCloudInterface->ClearOnWriteUserFileCompleteDelegates(this);
 		BaseUserCloudInterface->ClearOnWriteUserFileCanceledDelegates(this);
@@ -37,9 +34,6 @@ void FOnlineUserCloudEOSPlus::Initialize()
 	if (BaseUserCloudInterface.IsValid())
 	{
 		BaseUserCloudInterface->AddOnEnumerateUserFilesCompleteDelegate_Handle(FOnEnumerateUserFilesCompleteDelegate::CreateThreadSafeSP(this, &FOnlineUserCloudEOSPlus::OnEnumerateUserFilesComplete));
-PRAGMA_DISABLE_DEPRECATION_WARNINGS
-		BaseUserCloudInterface->AddOnWriteUserFileProgressDelegate_Handle(FOnWriteUserFileProgressDelegate::CreateThreadSafeSP(this, &FOnlineUserCloudEOSPlus::OnWriteUserFileProgress));
-PRAGMA_ENABLE_DEPRECATION_WARNINGS
 		BaseUserCloudInterface->AddOnWriteUserFileProgress64Delegate_Handle(FOnWriteUserFileProgress64Delegate::CreateThreadSafeSP(this, &FOnlineUserCloudEOSPlus::OnWriteUserFileProgress64));
 		BaseUserCloudInterface->AddOnWriteUserFileCompleteDelegate_Handle(FOnWriteUserFileCompleteDelegate::CreateThreadSafeSP(this, &FOnlineUserCloudEOSPlus::OnWriteUserFileComplete));
 		BaseUserCloudInterface->AddOnWriteUserFileCanceledDelegate_Handle(FOnWriteUserFileCanceledDelegate::CreateThreadSafeSP(this, &FOnlineUserCloudEOSPlus::OnWriteUserFileCanceled));
@@ -358,19 +352,6 @@ void FOnlineUserCloudEOSPlus::OnEnumerateUserFilesComplete(bool bWasSuccessful,
 	TriggerOnEnumerateUserFilesCompleteDelegates(bWasSuccessful, *NetIdPlus);
 }
 
-void FOnlineUserCloudEOSPlus::OnWriteUserFileProgress(int32 BytesWritten, const FUniqueNetId& UserId, const FString& FileName)
-{
-	const FUniqueNetIdEOSPlusPtr NetIdPlus = GetNetIdPlus(UserId.ToString());
-	if (!NetIdPlus.IsValid())
-	{
-		UE_LOG_ONLINE(Warning, TEXT("[FOnlineUserCloudEOSPlus::OnWriteUserFileProgress] Unknown user (%s)"), *UserId.ToString());
-	}
-
-PRAGMA_DISABLE_DEPRECATION_WARNINGS
-	TriggerOnWriteUserFileProgressDelegates(BytesWritten, *NetIdPlus, FileName);
-PRAGMA_ENABLE_DEPRECATION_WARNINGS
-}
-
 void FOnlineUserCloudEOSPlus::OnWriteUserFileProgress64(uint64 BytesWritten, const FUniqueNetId& UserId, const FString& FileName)
 {
 	const FUniqueNetIdEOSPlusPtr NetIdPlus = GetNetIdPlus(UserId.ToString());
diff --git a/Engine/Plugins/Online/OnlineSubsystemEOS/Source/OnlineSubsystemEOSPlus/Private/OnlineUserCloudEOSPlus.h b/Engine/Plugins/Online/OnlineSubsystemEOS/Source/OnlineSubsystemEOSPlus/Private/OnlineUserCloudEOSPlus.h
--- a/Engine/Plugins/Online/OnlineSubsystemEOS/Source/OnlineSubsystemEOSPlus/Private/OnlineUserCloudEOSPlus.h
+++ b/Engine/Plugins/Online/OnlineSubsystemEOS/Source/OnlineSubsystemEOSPlus/Private/OnlineUserCloudEOSPlus.h
@@ -42,7 +42,6 @@ PACKAGE_SCOPE:
 	FOnlineUserCloudEOSPlus(FOnlineSubsystemEOSPlus* InSubsystem);
 
 	void OnEnumerateUserFilesComplete(bool bWasSuccessful, const FUniqueNetId& UserId);
-	void OnWriteUserFileProgress(int32 BytesWritten, const FUniqueNetId& UserId, const FString& FileName);
 	void OnWriteUserFileProgress64(uint64 BytesWritten, const FUniqueNetId& UserId, const FString& FileName);
 	void OnWriteUserFileComplete(bool bWasSuccessful, const FUniqueNetId& UserId, const FString& FileName);
 	void OnWriteUserFileCanceled(bool bWasSuccessful, const FUniqueNetId& UserId, const FString& FileName);
diff --git a/Engine/Plugins/Online/OnlineSubsystemUtils/Source/OnlineSubsystemUtils/Classes/InAppPurchaseCallbackProxy2.h b/Engine/Plugins/Online/OnlineSubsystemUtils/Source/OnlineSubsystemUtils/Classes/InAppPurchaseCallbackProxy2.h
--- a/Engine/Plugins/Online/OnlineSubsystemUtils/Source/OnlineSubsystemUtils/Classes/InAppPurchaseCallbackProxy2.h
+++ b/Engine/Plugins/Online/OnlineSubsystemUtils/Source/OnlineSubsystemUtils/Classes/InAppPurchaseCallbackProxy2.h
@@ -102,9 +102,6 @@ private:
 	void OnPurchaseComplete();
 	void OnQueryReceiptsComplete(const FOnlineError& Result);
 
-	UE_DEPRECATED(5.3, "Use ::PurchaseStatusFromOnlineError instead")
-	EInAppPurchaseStatus PurchaseStatusFromOnlineError(const FOnlineError& OnlineError);
-
 	/** Triggers the In-App Purchase Transaction for the specifed user; the Purchase Request object must already be set up */
 	void Trigger(class APlayerController* PlayerController, const FInAppPurchaseProductRequest2& ProductRequest);
 	/** Triggers the In-App Purchase processing for unprocessed receipts */
diff --git a/Engine/Plugins/Online/OnlineSubsystemUtils/Source/OnlineSubsystemUtils/Classes/InAppPurchaseRestoreCallbackProxy2.h b/Engine/Plugins/Online/OnlineSubsystemUtils/Source/OnlineSubsystemUtils/Classes/InAppPurchaseRestoreCallbackProxy2.h
--- a/Engine/Plugins/Online/OnlineSubsystemUtils/Source/OnlineSubsystemUtils/Classes/InAppPurchaseRestoreCallbackProxy2.h
+++ b/Engine/Plugins/Online/OnlineSubsystemUtils/Source/OnlineSubsystemUtils/Classes/InAppPurchaseRestoreCallbackProxy2.h
@@ -66,8 +66,6 @@ private:
 	/** Triggers the In-App Purchase Restore Transaction for the specifed user */
 	void Trigger(const TArray<FInAppPurchaseProductRequest2>& ConsumableProductFlags, class APlayerController* PlayerController);
 
-	UE_DEPRECATED(5.3, "Use ::PurchaseStatusFromOnlineError instead")
-	EInAppPurchaseStatus PurchaseStatusFromOnlineError(const FOnlineError& OnlineError);
 private:
 
 	bool bWasSuccessful;
diff --git a/Engine/Plugins/Online/OnlineSubsystemUtils/Source/OnlineSubsystemUtils/Private/InAppPurchaseCallbackProxy2.cpp b/Engine/Plugins/Online/OnlineSubsystemUtils/Source/OnlineSubsystemUtils/Private/InAppPurchaseCallbackProxy2.cpp
--- a/Engine/Plugins/Online/OnlineSubsystemUtils/Source/OnlineSubsystemUtils/Private/InAppPurchaseCallbackProxy2.cpp
+++ b/Engine/Plugins/Online/OnlineSubsystemUtils/Source/OnlineSubsystemUtils/Private/InAppPurchaseCallbackProxy2.cpp
@@ -218,11 +218,6 @@ void UInAppPurchaseCallbackProxy2::OnPurchaseComplete()
 	}
 }
 
-EInAppPurchaseStatus UInAppPurchaseCallbackProxy2::PurchaseStatusFromOnlineError(const FOnlineError& OnlineError)
-{
-	return ::PurchaseStatusFromOnlineError(OnlineError);
-}
-
 void UInAppPurchaseCallbackProxy2::OnQueryReceiptsComplete(const FOnlineError& Result)
 {
 		SavedReceipts.Empty();
diff --git a/Engine/Plugins/Online/OnlineSubsystemUtils/Source/OnlineSubsystemUtils/Private/InAppPurchaseRestoreCallbackProxy2.cpp b/Engine/Plugins/Online/OnlineSubsystemUtils/Source/OnlineSubsystemUtils/Private/InAppPurchaseRestoreCallbackProxy2.cpp
--- a/Engine/Plugins/Online/OnlineSubsystemUtils/Source/OnlineSubsystemUtils/Private/InAppPurchaseRestoreCallbackProxy2.cpp
+++ b/Engine/Plugins/Online/OnlineSubsystemUtils/Source/OnlineSubsystemUtils/Private/InAppPurchaseRestoreCallbackProxy2.cpp
@@ -146,9 +146,4 @@ UInAppPurchaseRestoreCallbackProxy2* UInAppPurchaseRestoreCallbackProxy2::Create
 	Proxy->SetFlags(RF_StrongRefOnFrame);
 	Proxy->Trigger(ConsumableProductFlags, PlayerController);
 	return Proxy;
-}
-
-EInAppPurchaseStatus UInAppPurchaseRestoreCallbackProxy2::PurchaseStatusFromOnlineError(const FOnlineError& OnlineError)
-{
-	return ::PurchaseStatusFromOnlineError(OnlineError);
 }
\ No newline at end of file
diff --git a/Engine/Plugins/Online/OnlineSubsystemUtils/Source/OnlineSubsystemUtils/Private/OnlineBeaconClient.cpp b/Engine/Plugins/Online/OnlineSubsystemUtils/Source/OnlineSubsystemUtils/Private/OnlineBeaconClient.cpp
--- a/Engine/Plugins/Online/OnlineSubsystemUtils/Source/OnlineSubsystemUtils/Private/OnlineBeaconClient.cpp
+++ b/Engine/Plugins/Online/OnlineSubsystemUtils/Source/OnlineSubsystemUtils/Private/OnlineBeaconClient.cpp
@@ -563,14 +563,6 @@ void AOnlineBeaconClient::NotifyControlMessage(UNetConnection* Connection, uint8
 	}	
 }
 
-PRAGMA_DISABLE_DEPRECATION_WARNINGS
-FString AOnlineBeaconClient::GetAuthTicket(const FUniqueNetIdRepl& PlayerId)
-PRAGMA_ENABLE_DEPRECATION_WARNINGS
-{
-	return GetAuthTicketInternal(PlayerId);
-}
-
-
 FString AOnlineBeaconClient::GetAuthTicketInternal(const FUniqueNetIdRepl& PlayerId)
 {
 	FString AuthTicket;
diff --git a/Engine/Plugins/Online/OnlineSubsystemUtils/Source/OnlineSubsystemUtils/Private/OnlineBeaconHost.cpp b/Engine/Plugins/Online/OnlineSubsystemUtils/Source/OnlineSubsystemUtils/Private/OnlineBeaconHost.cpp
--- a/Engine/Plugins/Online/OnlineSubsystemUtils/Source/OnlineSubsystemUtils/Private/OnlineBeaconHost.cpp
+++ b/Engine/Plugins/Online/OnlineSubsystemUtils/Source/OnlineSubsystemUtils/Private/OnlineBeaconHost.cpp
@@ -315,32 +315,13 @@ bool AOnlineBeaconHost::HandleControlMessage(UNetConnection* Connection, uint8 M
 			// Try to kick off verification for this player.
 			const FString AuthTicket = UGameplayStatics::ParseOption(OptionsURL, TEXT("AuthTicket"));
 
-			bool bStartedAuth = false;
-
-			// Try to start deprecated auth method.
-			PRAGMA_DISABLE_DEPRECATION_WARNINGS
-			bStartedAuth = StartVerifyAuthentication(*UniqueIdRepl, AuthTicket);
-			PRAGMA_ENABLE_DEPRECATION_WARNINGS
-
-			if (!bStartedAuth)
+			// Create completion delegate.
+			FOnAuthenticationVerificationCompleteDelegate OnAuthComplete = FOnAuthenticationVerificationCompleteDelegate::CreateWeakLambda(this, [this, WeakConnection = TWeakObjectPtr<UNetConnection>(Connection)](const FOnlineError& OnlineError)
 			{
-				// Create completion delegate.
-				FOnAuthenticationVerificationCompleteDelegate OnAuthComplete = FOnAuthenticationVerificationCompleteDelegate::CreateWeakLambda(this, [this, WeakConnection = TWeakObjectPtr<UNetConnection>(Connection)](const FOnlineError& OnlineError)
-				{
-					OnAuthenticationVerificationComplete(WeakConnection.Get(), OnlineError);
-				});
+				OnAuthenticationVerificationComplete(WeakConnection.Get(), OnlineError);
+			});
 
-				// Try to start deprecated auth method.
-				PRAGMA_DISABLE_DEPRECATION_WARNINGS
-				bStartedAuth = StartVerifyAuthentication(*UniqueIdRepl, AuthTicket, OnAuthComplete);
-				PRAGMA_ENABLE_DEPRECATION_WARNINGS
-
-				// Don't start new auth method if deprecated auth is active.
-				if (!bStartedAuth)
-				{
-					bStartedAuth = StartVerifyAuthentication(*UniqueIdRepl, OptionsURL, AuthTicket, OnAuthComplete);
-				}
-			}
+			bool bStartedAuth = StartVerifyAuthentication(*UniqueIdRepl, OptionsURL, AuthTicket, OnAuthComplete);
 
 			if (!bStartedAuth)
 			{
@@ -659,32 +640,6 @@ void AOnlineBeaconHost::RemoveClientActor(AOnlineBeaconClient* ClientActor)
 	}
 }
 
-PRAGMA_DISABLE_DEPRECATION_WARNINGS
-bool AOnlineBeaconHost::StartVerifyAuthentication(const FUniqueNetId& PlayerId, const FString& AuthenticationToken)
-PRAGMA_ENABLE_DEPRECATION_WARNINGS
-{
-	return false;
-}
-
-PRAGMA_DISABLE_DEPRECATION_WARNINGS
-void AOnlineBeaconHost::OnAuthenticationVerificationComplete(const class FUniqueNetId& PlayerId, const FOnlineError& Error)
-PRAGMA_ENABLE_DEPRECATION_WARNINGS
-{
-	UNetConnection* Connection = nullptr;
-	FConnectionState* ConnState = nullptr;
-	if (GetConnectionDataForUniqueNetId(PlayerId, Connection, ConnState))
-	{
-		OnAuthenticationVerificationComplete(Connection, Error);
-	}
-}
-
-PRAGMA_DISABLE_DEPRECATION_WARNINGS
-bool AOnlineBeaconHost::StartVerifyAuthentication(const FUniqueNetId& PlayerId, const FString& AuthenticationToken, const FOnAuthenticationVerificationCompleteDelegate& OnComplete)
-PRAGMA_ENABLE_DEPRECATION_WARNINGS
-{
-	return false;
-}
-
 bool AOnlineBeaconHost::StartVerifyAuthentication(const FUniqueNetId& PlayerId, const FString& LoginOptions, const FString& AuthenticationToken, const FOnAuthenticationVerificationCompleteDelegate& OnComplete)
 {
 	return false;
diff --git a/Engine/Plugins/Online/OnlineSubsystemUtils/Source/OnlineSubsystemUtils/Private/Tests/OnlineBeaconUnitTestHost.cpp b/Engine/Plugins/Online/OnlineSubsystemUtils/Source/OnlineSubsystemUtils/Private/Tests/OnlineBeaconUnitTestHost.cpp
--- a/Engine/Plugins/Online/OnlineSubsystemUtils/Source/OnlineSubsystemUtils/Private/Tests/OnlineBeaconUnitTestHost.cpp
+++ b/Engine/Plugins/Online/OnlineSubsystemUtils/Source/OnlineSubsystemUtils/Private/Tests/OnlineBeaconUnitTestHost.cpp
@@ -14,71 +14,6 @@ AOnlineBeaconUnitTestHost::AOnlineBeaconUnitTestHost(const FObjectInitializer& O
 
 #if WITH_DEV_AUTOMATION_TESTS && WITH_EDITOR
 
-bool AOnlineBeaconUnitTestHost::StartVerifyAuthentication(const FUniqueNetId& PlayerId, const FString& AuthenticationToken)
-{
-	const BeaconUnitTest::FTestConfig* TestConfig = BeaconUnitTest::FTestPrerequisites::GetActiveTestConfig();
-	if (TestConfig == nullptr)
-	{
-		return false;
-	}
-
-	if (!TestConfig->Auth.Method1.bEnabled)
-	{
-		PRAGMA_DISABLE_DEPRECATION_WARNINGS
-		return Super::StartVerifyAuthentication(PlayerId, AuthenticationToken);
-		PRAGMA_ENABLE_DEPRECATION_WARNINGS
-	}
-
-	if (TestConfig->Auth.Method1.bDelayDelegate)
-	{
-		return BeaconUnitTest::SetTimerForNextFrame(GetWorld(), GFrameCounter, [this, PlayerIdRef = PlayerId.AsShared()]()
-		{
-			const BeaconUnitTest::FTestConfig* TestConfig = BeaconUnitTest::FTestPrerequisites::GetActiveTestConfig();
-			PRAGMA_DISABLE_DEPRECATION_WARNINGS
-			OnAuthenticationVerificationComplete(*PlayerIdRef, TestConfig ? TestConfig->Auth.Method1.Result : FOnlineError());
-			PRAGMA_ENABLE_DEPRECATION_WARNINGS
-		});
-	}
-	else
-	{
-		PRAGMA_DISABLE_DEPRECATION_WARNINGS
-		OnAuthenticationVerificationComplete(PlayerId, TestConfig->Auth.Method1.Result);
-		PRAGMA_ENABLE_DEPRECATION_WARNINGS
-
-		return true;
-	}
-}
-
-bool AOnlineBeaconUnitTestHost::StartVerifyAuthentication(const FUniqueNetId& PlayerId, const FString& AuthenticationToken, const FOnAuthenticationVerificationCompleteDelegate& OnComplete)
-{
-	const BeaconUnitTest::FTestConfig* TestConfig = BeaconUnitTest::FTestPrerequisites::GetActiveTestConfig();
-	if (TestConfig == nullptr)
-	{
-		return false;
-	}
-
-	if (!TestConfig->Auth.Method2.bEnabled)
-	{
-		PRAGMA_DISABLE_DEPRECATION_WARNINGS
-		return Super::StartVerifyAuthentication(PlayerId, AuthenticationToken, OnComplete);
-		PRAGMA_ENABLE_DEPRECATION_WARNINGS
-	}
-
-	if (TestConfig->Auth.Method2.bDelayDelegate)
-	{
-		return BeaconUnitTest::SetTimerForNextFrame(GetWorld(), GFrameCounter, [this, OnComplete]()
-		{
-			const BeaconUnitTest::FTestConfig* TestConfig = BeaconUnitTest::FTestPrerequisites::GetActiveTestConfig();
-			OnComplete.ExecuteIfBound(TestConfig ? TestConfig->Auth.Method2.Result : FOnlineError());
-		});
-	}
-	else
-	{
-		OnComplete.ExecuteIfBound(TestConfig->Auth.Method2.Result);
-		return true;
-	}
-}
-
 bool AOnlineBeaconUnitTestHost::StartVerifyAuthentication(const FUniqueNetId& PlayerId, const FString& LoginOptions, const FString& AuthenticationToken, const FOnAuthenticationVerificationCompleteDelegate& OnComplete)
 {
 	const BeaconUnitTest::FTestConfig* TestConfig = BeaconUnitTest::FTestPrerequisites::GetActiveTestConfig();
diff --git a/Engine/Plugins/Online/OnlineSubsystemUtils/Source/OnlineSubsystemUtils/Private/Tests/OnlineBeaconUnitTestHost.h b/Engine/Plugins/Online/OnlineSubsystemUtils/Source/OnlineSubsystemUtils/Private/Tests/OnlineBeaconUnitTestHost.h
--- a/Engine/Plugins/Online/OnlineSubsystemUtils/Source/OnlineSubsystemUtils/Private/Tests/OnlineBeaconUnitTestHost.h
+++ b/Engine/Plugins/Online/OnlineSubsystemUtils/Source/OnlineSubsystemUtils/Private/Tests/OnlineBeaconUnitTestHost.h
@@ -21,12 +21,6 @@ public:
 #if WITH_DEV_AUTOMATION_TESTS && WITH_EDITOR
 
 	//~ Begin AOnlineBeaconHost Interface
-	PRAGMA_DISABLE_DEPRECATION_WARNINGS
-	virtual bool StartVerifyAuthentication(const FUniqueNetId& PlayerId, const FString& AuthenticationToken) override;
-	PRAGMA_ENABLE_DEPRECATION_WARNINGS
-	PRAGMA_DISABLE_DEPRECATION_WARNINGS
-	virtual bool StartVerifyAuthentication(const FUniqueNetId& PlayerId, const FString& AuthenticationToken, const FOnAuthenticationVerificationCompleteDelegate& OnComplete) override;
-	PRAGMA_ENABLE_DEPRECATION_WARNINGS
 	virtual bool StartVerifyAuthentication(const FUniqueNetId& PlayerId, const FString& LoginOptions, const FString& AuthenticationToken, const FOnAuthenticationVerificationCompleteDelegate& OnComplete) override;
 	virtual bool VerifyJoinForBeaconType(const FUniqueNetId& PlayerId, const FString& BeaconType) override;
 	//~ End AOnlineBeaconHost Interface
diff --git a/Engine/Plugins/Online/OnlineSubsystemUtils/Source/OnlineSubsystemUtils/Public/OnlineBeaconClient.h b/Engine/Plugins/Online/OnlineSubsystemUtils/Source/OnlineSubsystemUtils/Public/OnlineBeaconClient.h
--- a/Engine/Plugins/Online/OnlineSubsystemUtils/Source/OnlineSubsystemUtils/Public/OnlineBeaconClient.h
+++ b/Engine/Plugins/Online/OnlineSubsystemUtils/Source/OnlineSubsystemUtils/Public/OnlineBeaconClient.h
@@ -114,16 +114,6 @@ class ONLINESUBSYSTEMUTILS_API AOnlineBeaconClient : public AOnlineBeacon
 	 */
 	const FUniqueNetIdRepl& GetUniqueId() const;
 
-	/**
-	 * Get the authentication ticket to send when the server requires auth.
-	 * 
-	 * @param PlayerId unique id of the user on this connection
-	 * 
-	 * @return AuthTicket value to send to server
-	 */
-	UE_DEPRECATED(5.3, "GetAuthTicket is deprecated. Override the login options by overriding GetLoginOptions and specifying AuthTicket=%s with the resulting string")
-	virtual FString GetAuthTicket(const FUniqueNetIdRepl& PlayerId);
-
 	/**
 	 * Get the login options to send the server during login.
 	 * Typically this is expected to be in the format that UGameplayStatics::ParseOption can parse, i.e. Key1=Value1?Key2=Value2?Key3=Value3
diff --git a/Engine/Plugins/Online/OnlineSubsystemUtils/Source/OnlineSubsystemUtils/Public/OnlineBeaconHost.h b/Engine/Plugins/Online/OnlineSubsystemUtils/Source/OnlineSubsystemUtils/Public/OnlineBeaconHost.h
--- a/Engine/Plugins/Online/OnlineSubsystemUtils/Source/OnlineSubsystemUtils/Public/OnlineBeaconHost.h
+++ b/Engine/Plugins/Online/OnlineSubsystemUtils/Source/OnlineSubsystemUtils/Public/OnlineBeaconHost.h
@@ -122,25 +122,6 @@ protected:
 	UPROPERTY(Config)
 	uint32 MaxAuthTokenSize = 1024;
 
-	/**
-	 * Start verifying an authentication token for a connection.
-	 * OnAuthenticationVerificationComplete must be called to complete authentication verification.
-	 *
-	 * @param PlayerId net id of player to authenticate.
-	 * @param AuthenticationToken token to use for verification.
-	 */
-	 UE_DEPRECATED(5.2, "This version of the StartVerifyAuthentication is deprecated. Please use the new StartVerifyAuthentication method instead.")
-	virtual bool StartVerifyAuthentication(const FUniqueNetId& PlayerId, const FString& AuthenticationToken);
-
-	/**
-	 * Event which must be signaled to complete an authentication verification request.
-	 *
-	 * @param PlayerId net id of player to authenticate.
-	 * @param Error result of the operation.
-	 */
-	UE_DEPRECATED(5.2, "This version of the OnAuthenticationVerificationComplete is deprecated. Please use the new OnAuthenticationVerificationComplete method instead.")
-	void OnAuthenticationVerificationComplete(const class FUniqueNetId& PlayerId, const FOnlineError& Error);
-
 	/**
 	 * Delegate executed when user authentication has completed.
 	 *
@@ -148,17 +129,6 @@ protected:
 	 */
 	DECLARE_DELEGATE_OneParam(FOnAuthenticationVerificationCompleteDelegate, const FOnlineError& /*OnlineError*/);
 
-	/**
-	 * Start verifying an authentication token for a connection.
-	 * OnAuthenticationVerificationComplete must be called to complete authentication verification.
-	 *
-	 * @param NetConnection network connection associated with the authentication challenge.
-	 * @param PlayerId net id of player to authenticate.
-	 * @param AuthenticationToken token to use for verification.
-	 */
-	UE_DEPRECATED(5.3, "This version of the StartVerifyAuthentication is deprecated. Please use the new StartVerifyAuthentication method instead.")
-	virtual bool StartVerifyAuthentication(const FUniqueNetId& PlayerId, const FString& AuthenticationToken, const FOnAuthenticationVerificationCompleteDelegate& OnComplete);
-
 	/**
 	 * Start verifying an authentication token for a connection.
 	 * OnAuthenticationVerificationComplete must be called to complete authentication verification.
diff --git a/Engine/Plugins/Online/VoiceChat/VoiceChat/Source/Public/VoiceChat.h b/Engine/Plugins/Online/VoiceChat/VoiceChat/Source/Public/VoiceChat.h
--- a/Engine/Plugins/Online/VoiceChat/VoiceChat/Source/Public/VoiceChat.h
+++ b/Engine/Plugins/Online/VoiceChat/VoiceChat/Source/Public/VoiceChat.h
@@ -48,8 +48,6 @@ enum class EVoiceChatTransmitMode
 	None,
 	/** Transmit to all of the channels you are in */
 	All,
-	/** Transmit to a specific channel */
-	Channel UE_DEPRECATED(5.3, "EVoiceChatTransmitMode::Channel is deprecated. Please use EVoiceChatTransmitMode::SpecificChannels instead!"),
 	/** Transmit to specific channel(s) */
 	SpecificChannels,
 };
@@ -587,14 +585,6 @@ public:
 	 */
 	virtual void TransmitToNoChannels() = 0;
 
-	/**
-	 * Transmit to a specific channel
-	 *
-	 * @param ChannelName Channel to transmit to
-	 */
-	UE_DEPRECATED(5.3, "TransmitToSpecificChannel is deprecated. Please use TransmitToSpecificChannels instead!")
-	virtual void TransmitToSpecificChannel(const FString& ChannelName) { TransmitToSpecificChannels({ChannelName}); }
-
 	/**
 	 * Transmit to Specified set of channels
 	 * @param ChannelNames Channels to transmit to
@@ -620,14 +610,6 @@ public:
 	 */
 	virtual EVoiceChatTransmitMode GetTransmitMode() const = 0;
 
-	/**
-	 * Get the first channel we are transmitting to
-	 *
-	 * @return Channel we are transmitting to. Will be empty if GetTransmitMode() returned a value other than EVoiceChatTransmitMode::SpecificChannels
-	 */
-	UE_DEPRECATED(5.3, "GetTransmitChannel is deprecated. Please use GetTransmitChannels instead!")
-		virtual FString GetTransmitChannel() const { return ((GetTransmitMode() == EVoiceChatTransmitMode::SpecificChannels) && !GetTransmitChannels().IsEmpty()) ? GetTransmitChannels().Array()[0] : FString(); };
-
 	/**
 	 * Get all channels we are transmitting to
 	 *
diff --git a/Engine/Source/Runtime/Online/BuildPatchServices/Private/Tests/Mock/HttpRequest.mock.h b/Engine/Source/Runtime/Online/BuildPatchServices/Private/Tests/Mock/HttpRequest.mock.h
--- a/Engine/Source/Runtime/Online/BuildPatchServices/Private/Tests/Mock/HttpRequest.mock.h
+++ b/Engine/Source/Runtime/Online/BuildPatchServices/Private/Tests/Mock/HttpRequest.mock.h
@@ -176,11 +176,6 @@ namespace BuildPatchServices
 			return HttpRequestCompleteDelegate;
 		}
 
-		virtual FHttpRequestProgressDelegate& OnRequestProgress() override
-		{
-			return HttpRequestProgressDelegate;
-		}
-
 		virtual FHttpRequestProgressDelegate64& OnRequestProgress64() override
 		{
 			return HttpRequestProgressDelegate64;
diff --git a/Engine/Source/Runtime/Online/HTTP/Private/Curl/CurlHttp.cpp b/Engine/Source/Runtime/Online/HTTP/Private/Curl/CurlHttp.cpp
--- a/Engine/Source/Runtime/Online/HTTP/Private/Curl/CurlHttp.cpp
+++ b/Engine/Source/Runtime/Online/HTTP/Private/Curl/CurlHttp.cpp
@@ -1170,7 +1170,6 @@ void FCurlHttpRequest::CheckProgressDelegate()
 		LastReportedBytesSent = CurrentBytesSent;
 		LastReportedBytesRead = CurrentBytesRead;
 		// Update response progress
-		OnRequestProgress().ExecuteIfBound(SharedThis(this), LastReportedBytesSent, LastReportedBytesRead);
 		OnRequestProgress64().ExecuteIfBound(SharedThis(this), LastReportedBytesSent, LastReportedBytesRead);
 	}
 }
diff --git a/Engine/Source/Runtime/Online/HTTP/Private/GenericPlatform/HttpRequestImpl.cpp b/Engine/Source/Runtime/Online/HTTP/Private/GenericPlatform/HttpRequestImpl.cpp
--- a/Engine/Source/Runtime/Online/HTTP/Private/GenericPlatform/HttpRequestImpl.cpp
+++ b/Engine/Source/Runtime/Online/HTTP/Private/GenericPlatform/HttpRequestImpl.cpp
@@ -9,13 +9,6 @@ FHttpRequestCompleteDelegate& FHttpRequestImpl::OnProcessRequestComplete()
 	return RequestCompleteDelegate;
 }
 
-PRAGMA_DISABLE_DEPRECATION_WARNINGS
-FHttpRequestProgressDelegate& FHttpRequestImpl::OnRequestProgress() 
-{
-	return RequestProgressDelegate;
-}
-PRAGMA_ENABLE_DEPRECATION_WARNINGS
-
 FHttpRequestProgressDelegate64& FHttpRequestImpl::OnRequestProgress64() 
 {
 	return RequestProgressDelegate64;
@@ -39,9 +32,6 @@ FHttpRequestWillRetryDelegate& FHttpRequestImpl::OnRequestWillRetry()
 void FHttpRequestImpl::Shutdown()
 {
 	OnProcessRequestComplete().Unbind();
-PRAGMA_DISABLE_DEPRECATION_WARNINGS
-	OnRequestProgress().Unbind();
-PRAGMA_ENABLE_DEPRECATION_WARNINGS
 	OnRequestProgress64().Unbind();
 	OnStatusCodeReceived().Unbind();
 	OnHeaderReceived().Unbind();
diff --git a/Engine/Source/Runtime/Online/HTTP/Private/HttpManager.cpp b/Engine/Source/Runtime/Online/HTTP/Private/HttpManager.cpp
--- a/Engine/Source/Runtime/Online/HTTP/Private/HttpManager.cpp
+++ b/Engine/Source/Runtime/Online/HTTP/Private/HttpManager.cpp
@@ -244,8 +244,6 @@ bool FHttpManager::IsDomainAllowed(const FString& Url) const
 		return URLRequestFilter.IsRequestAllowed(Url);
 	}
 
-PRAGMA_DISABLE_DEPRECATION_WARNINGS
-
 #if !UE_BUILD_SHIPPING
 #if !(UE_GAME || UE_SERVER)
 	// Allowed domain filtering is opt-in in non-shipping non-game/server builds
@@ -264,23 +262,7 @@ PRAGMA_DISABLE_DEPRECATION_WARNINGS
 #endif
 #endif // !UE_BUILD_SHIPPING
 
-	// Check to see if the Domain is allowed (either on the list or the list was empty)
-	const TArray<FString>& AllowedDomains = FHttpModule::Get().GetAllowedDomains();
-	if (AllowedDomains.Num() > 0)
-	{
-		const FString Domain = FPlatformHttp::GetUrlDomain(Url);
-		for (const FString& AllowedDomain : AllowedDomains)
-		{
-			if (Domain.EndsWith(AllowedDomain))
-			{
-				return true;
-			}
-		}
-		return false;
-	}
 	return true;
-
-PRAGMA_ENABLE_DEPRECATION_WARNINGS
 }
 
 /*static*/
diff --git a/Engine/Source/Runtime/Online/HTTP/Private/HttpRetrySystem.cpp b/Engine/Source/Runtime/Online/HTTP/Private/HttpRetrySystem.cpp
--- a/Engine/Source/Runtime/Online/HTTP/Private/HttpRetrySystem.cpp
+++ b/Engine/Source/Runtime/Online/HTTP/Private/HttpRetrySystem.cpp
@@ -203,9 +203,6 @@ void FHttpRetrySystem::FRequest::CancelRequest()
 
 void FHttpRetrySystem::FRequest::HttpOnRequestProgress(FHttpRequestPtr InHttpRequest, uint64 BytesSent, uint64 BytesRcv)
 {
-PRAGMA_DISABLE_DEPRECATION_WARNINGS
-	OnRequestProgress().ExecuteIfBound(AsShared(), BytesSent, BytesRcv);
-PRAGMA_ENABLE_DEPRECATION_WARNINGS
 	OnRequestProgress64().ExecuteIfBound(AsShared(), BytesSent, BytesRcv);
 }
 
diff --git a/Engine/Source/Runtime/Online/HTTP/Public/GenericPlatform/HttpRequestImpl.h b/Engine/Source/Runtime/Online/HTTP/Public/GenericPlatform/HttpRequestImpl.h
--- a/Engine/Source/Runtime/Online/HTTP/Public/GenericPlatform/HttpRequestImpl.h
+++ b/Engine/Source/Runtime/Online/HTTP/Public/GenericPlatform/HttpRequestImpl.h
@@ -13,7 +13,6 @@ class FHttpRequestImpl : public IHttpRequest
 public:
 	// IHttpRequest
 	HTTP_API virtual FHttpRequestCompleteDelegate& OnProcessRequestComplete() override;
-	HTTP_API virtual FHttpRequestProgressDelegate& OnRequestProgress() override;
 	HTTP_API virtual FHttpRequestProgressDelegate64& OnRequestProgress64() override;
 	HTTP_API virtual FHttpRequestStatusCodeReceivedDelegate& OnStatusCodeReceived() override;
 	HTTP_API virtual FHttpRequestHeaderReceivedDelegate& OnHeaderReceived() override;
@@ -35,9 +34,6 @@ protected:
 	/** Delegate that will get called once request completes or on any error */
 	FHttpRequestCompleteDelegate RequestCompleteDelegate;
 
-	/** Delegate that will get called once per tick with bytes downloaded so far */
-	FHttpRequestProgressDelegate RequestProgressDelegate;
-
 	/** Delegate that will get called once per tick with bytes downloaded so far */
 	FHttpRequestProgressDelegate64 RequestProgressDelegate64;
 
diff --git a/Engine/Source/Runtime/Online/HTTP/Public/HttpModule.h b/Engine/Source/Runtime/Online/HTTP/Public/HttpModule.h
--- a/Engine/Source/Runtime/Online/HTTP/Public/HttpModule.h
+++ b/Engine/Source/Runtime/Online/HTTP/Public/HttpModule.h
@@ -326,15 +326,6 @@ public:
 		return bSupportsDynamicProxy;
 	}
 
-	/**
-	 * @returns the list of domains allowed to be visited in a shipping build
-	 */
-	UE_DEPRECATED(5.3, "GetAllowedDomains has been deprecated. URLRequestFilter should be used instead.")
-	inline const TArray<FString>& GetAllowedDomains() const
-	{
-		return AllowedDomains;
-	}
-
 protected:
 	/** timeout in seconds to establish the connection */
 	float HttpConnectionTimeout;
diff --git a/Engine/Source/Runtime/Online/HTTP/Public/Interfaces/IHttpRequest.h b/Engine/Source/Runtime/Online/HTTP/Public/Interfaces/IHttpRequest.h
--- a/Engine/Source/Runtime/Online/HTTP/Public/Interfaces/IHttpRequest.h
+++ b/Engine/Source/Runtime/Online/HTTP/Public/Interfaces/IHttpRequest.h
@@ -358,12 +358,6 @@ public:
 	 */
 	virtual FHttpRequestCompleteDelegate& OnProcessRequestComplete() = 0;
 
-	/**
-	 * Delegate called to update the request/response progress. See FHttpRequestProgressDelegate
-	 */
-	UE_DEPRECATED(5.3, "OnRequestProgress has been deprecated, use OnRequestProgress64 instead")
-	virtual FHttpRequestProgressDelegate& OnRequestProgress() = 0;
-
 	/**
 	 * Delegate called to update the request/response progress. See FHttpRequestProgressDelegate64
 	 */
