Removal of Online logic marked for deprecation in 5.2

diff --git a/Engine/Plugins/Online/OnlineServicesEOSGS/Source/Private/Online/LobbiesEOSGSTypes.cpp b/Engine/Plugins/Online/OnlineServicesEOSGS/Source/Private/Online/LobbiesEOSGSTypes.cpp
--- a/Engine/Plugins/Online/OnlineServicesEOSGS/Source/Private/Online/LobbiesEOSGSTypes.cpp
+++ b/Engine/Plugins/Online/OnlineServicesEOSGS/Source/Private/Online/LobbiesEOSGSTypes.cpp
@@ -42,10 +42,6 @@ FString ToLogString(const FLobbyBucketIdEOS& BucketId)
 	return FString::Printf(TEXT("%s:%d"), *BucketId.GetProductName(), BucketId.GetProductVersion());
 }
 
-PRAGMA_DISABLE_DEPRECATION_WARNINGS
-const FString FLobbyBucketIdEOS::Separator = TEXT("|");
-PRAGMA_ENABLE_DEPRECATION_WARNINGS
-
 FLobbyBucketIdEOS::FLobbyBucketIdEOS(FString ProductName, int32 ProductVersion)
 	: ProductName(ProductName.Replace(UE_ONLINE_LOBBIES_BUCKET_ID_EOS_SEPARATOR, TEXT("_")))
 	, ProductVersion(ProductVersion)
diff --git a/Engine/Plugins/Online/OnlineServicesEOSGS/Source/Private/Online/LobbiesEOSGSTypes.h b/Engine/Plugins/Online/OnlineServicesEOSGS/Source/Private/Online/LobbiesEOSGSTypes.h
--- a/Engine/Plugins/Online/OnlineServicesEOSGS/Source/Private/Online/LobbiesEOSGSTypes.h
+++ b/Engine/Plugins/Online/OnlineServicesEOSGS/Source/Private/Online/LobbiesEOSGSTypes.h
@@ -37,9 +37,6 @@ struct FLobbyMemberServiceSnapshot;
 class FLobbyBucketIdEOS
 {
 public:
-	UE_DEPRECATED(5.2, "Separator is deprecated, please use UE_ONLINE_LOBBIES_BUCKET_ID_EOS_SEPARATOR instead.")
-	static const FString Separator;
-
 	FLobbyBucketIdEOS() = default;
 	FLobbyBucketIdEOS(const FLobbyBucketIdEOS&) = default;
 	FLobbyBucketIdEOS(FLobbyBucketIdEOS&&) = default;
diff --git a/Engine/Plugins/Online/OnlineSubsystem/Source/Private/OnlineError.cpp b/Engine/Plugins/Online/OnlineSubsystem/Source/Private/OnlineError.cpp
--- a/Engine/Plugins/Online/OnlineSubsystem/Source/Private/OnlineError.cpp
+++ b/Engine/Plugins/Online/OnlineSubsystem/Source/Private/OnlineError.cpp
@@ -5,10 +5,6 @@
 #define LOCTEXT_NAMESPACE "OnlineError"
 #define ONLINE_ERROR_NAMESPACE "errors.com.epicgames.online.generic"
 
-PRAGMA_DISABLE_DEPRECATION_WARNINGS
-const FString FOnlineError::GenericErrorCode = TEXT("GenericError");
-PRAGMA_ENABLE_DEPRECATION_WARNINGS
-
 const FOnlineError& FOnlineError::Success() { static FOnlineError Error(EOnlineErrorResult::Success); return Error; }
 
 #if ONLINE_ERROR_LEGACY
diff --git a/Engine/Plugins/Online/OnlineSubsystem/Source/Private/OnlineSubsystem.cpp b/Engine/Plugins/Online/OnlineSubsystem/Source/Private/OnlineSubsystem.cpp
--- a/Engine/Plugins/Online/OnlineSubsystem/Source/Private/OnlineSubsystem.cpp
+++ b/Engine/Plugins/Online/OnlineSubsystem/Source/Private/OnlineSubsystem.cpp
@@ -127,18 +127,10 @@ const FName TENCENT_SUBSYSTEM(TEXT("TENCENT"));
 const FName SWITCH_SUBSYSTEM(TEXT("SWITCH"));
 #endif
 
-PRAGMA_DISABLE_DEPRECATION_WARNINGS
-const FName OCULUS_SUBSYSTEM(TEXT("OCULUS"));
-PRAGMA_ENABLE_DEPRECATION_WARNINGS
-
 #ifndef SAMSUNG_SUBSYSTEM
 const FName SAMSUNG_SUBSYSTEM(TEXT("SAMSUNG"));
 #endif
 
-#ifndef QUAIL_SUBSYSTEM
-const FName QUAIL_SUBSYSTEM(TEXT("Quail"));
-#endif
-
 #ifndef EOS_SUBSYSTEM
 const FName EOS_SUBSYSTEM(TEXT("EOS"));
 #endif
diff --git a/Engine/Plugins/Online/OnlineSubsystem/Source/Public/Interfaces/OnlineIdentityInterface.h b/Engine/Plugins/Online/OnlineSubsystem/Source/Public/Interfaces/OnlineIdentityInterface.h
--- a/Engine/Plugins/Online/OnlineSubsystem/Source/Public/Interfaces/OnlineIdentityInterface.h
+++ b/Engine/Plugins/Online/OnlineSubsystem/Source/Public/Interfaces/OnlineIdentityInterface.h
@@ -478,12 +478,6 @@ public:
 		Delegate.ExecuteIfBound(LocalUserNum, false, EmptyToken);
 	}
 
-	UE_DEPRECATED(5.2, "Please use GetLinkedAccountAuthToken taking a TokenType")
-	virtual void GetLinkedAccountAuthToken(int32 LocalUserNum, const FOnGetLinkedAccountAuthTokenCompleteDelegate& Delegate) const
-	{
-		GetLinkedAccountAuthToken(LocalUserNum, FString(), Delegate);
-	}
-
 	/**
 	 * Delegate executed when we get a user privilege result.
 	 *
diff --git a/Engine/Plugins/Online/OnlineSubsystem/Source/Public/Interfaces/OnlineSessionInterface.h b/Engine/Plugins/Online/OnlineSubsystem/Source/Public/Interfaces/OnlineSessionInterface.h
--- a/Engine/Plugins/Online/OnlineSubsystem/Source/Public/Interfaces/OnlineSessionInterface.h
+++ b/Engine/Plugins/Online/OnlineSubsystem/Source/Public/Interfaces/OnlineSessionInterface.h
@@ -75,30 +75,6 @@ inline const TCHAR* LexToString(const EOnJoinSessionCompleteResult::Type Value)
 DECLARE_MULTICAST_DELEGATE_TwoParams(FOnJoinSessionComplete, FName, EOnJoinSessionCompleteResult::Type);
 typedef FOnJoinSessionComplete::FDelegate FOnJoinSessionCompleteDelegate;
 
-/**
- * Delegate fired when a player has joined or left a session
- *
- * @param SessionName The name of the session that changed
- * @param UniqueId The ID of the user whose join state has changed
- * @param bJoined If true this is a join event, (if false it is a leave event)
- */
-UE_DEPRECATED(5.2, "FOnSessionParticipantsChange is deprecated, please use FOnSessionParticipantJoined and FOnSessionParticipantLeft instead.")
-DECLARE_MULTICAST_DELEGATE_ThreeParams(FOnSessionParticipantsChange, FName, const FUniqueNetId&, bool);
-PRAGMA_DISABLE_DEPRECATION_WARNINGS
-typedef FOnSessionParticipantsChange::FDelegate FOnSessionParticipantsChangeDelegate;
-PRAGMA_ENABLE_DEPRECATION_WARNINGS
-
-/**
- * Delegate fired when a session's member is removed from the session
- * @param SessionName The name of the session
- * @param TargetUniqueNetId The UniqueNetId of the member who was removed
- */
-UE_DEPRECATED(5.2, "FOnSessionParticipantRemoved is deprecated, please use FOnSessionParticipantLeft instead.")
-DECLARE_MULTICAST_DELEGATE_TwoParams(FOnSessionParticipantRemoved, FName, const FUniqueNetId&);
-PRAGMA_DISABLE_DEPRECATION_WARNINGS
-typedef FOnSessionParticipantRemoved::FDelegate FOnSessionParticipantRemovedDelegate;
-PRAGMA_ENABLE_DEPRECATION_WARNINGS
-
 /**
  * Delegate fired when a player has joined a session
  *
@@ -659,24 +635,6 @@ public:
 	 */
 	DEFINE_ONLINE_DELEGATE_TWO_PARAM(OnJoinSessionComplete, FName, EOnJoinSessionCompleteResult::Type);
 
-	PRAGMA_DISABLE_DEPRECATION_WARNINGS
-	/**
-	 * Delegate fired when a player has joined or left a session
-	 * @param SessionName The name of the session that changed
-	 * @param UniqueId The ID of the user whose join state has changed
-	 * @param bJoined if true this is a join event, (if false it is a leave event)
-	 */
-	DEFINE_ONLINE_DELEGATE_THREE_PARAM(OnSessionParticipantsChange, FName, const FUniqueNetId&, bool);
-
-	/**
-	* Delegate fired when a session's member is removed from the session
-	*
-	* @param SessionName The name of the session
-	* @param TargetUniqueNetId The UniqueNetId of the member who was removed
-	*/
-	DEFINE_ONLINE_DELEGATE_TWO_PARAM(OnSessionParticipantRemoved, FName, const FUniqueNetId&);
-	PRAGMA_ENABLE_DEPRECATION_WARNINGS
-
 	/**
 	* Delegate fired when a player joins a session
 	* @param SessionName The name of the session
diff --git a/Engine/Plugins/Online/OnlineSubsystem/Source/Public/OnlineError.h b/Engine/Plugins/Online/OnlineSubsystem/Source/Public/OnlineError.h
--- a/Engine/Plugins/Online/OnlineSubsystem/Source/Public/OnlineError.h
+++ b/Engine/Plugins/Online/OnlineSubsystem/Source/Public/OnlineError.h
@@ -111,10 +111,6 @@ public:
 	/** Setters for adding the raw error */
 	inline FOnlineError& SetErrorRaw(const FString& Val) { ErrorRaw = Val; return *this; }
 
-	/** Code useful when all you have is raw error info from old APIs */
-	UE_DEPRECATED(5.2, "Access of GenericErrorCode is deprecated, use GetGenericErrorCode() instead.")
-	static const FString GenericErrorCode;
-
 	/** Code useful when all you have is raw error info from old APIs */
 	static const FString GetGenericErrorCode();
 
diff --git a/Engine/Plugins/Online/OnlineSubsystem/Source/Public/OnlineIdentityErrors.h b/Engine/Plugins/Online/OnlineSubsystem/Source/Public/OnlineIdentityErrors.h
--- a/Engine/Plugins/Online/OnlineSubsystem/Source/Public/OnlineIdentityErrors.h
+++ b/Engine/Plugins/Online/OnlineSubsystem/Source/Public/OnlineIdentityErrors.h
@@ -26,23 +26,6 @@ namespace OnlineIdentity
 		inline FOnlineError PinGrantTimeout() { return ONLINE_ERROR(EOnlineErrorResult::FailExtended, TEXT("pin_grant_timeout")); }
 
 		inline FOnlineError Canceled() { return ONLINE_ERROR(EOnlineErrorResult::Canceled); }
-
-		// Params
-		UE_DEPRECATED(5.2, "AuthLoginParam is deprecated, please use UE_ONLINE_ERROR_PARAM_IDENTITY_AUTH_LOGIN instead.")
-		extern ONLINESUBSYSTEM_API const FString AuthLoginParam;
-
-		UE_DEPRECATED(5.2, "AuthTypeParam is deprecated, please use UE_ONLINE_ERROR_PARAM_IDENTITY_AUTH_TYPE instead.")
-		extern ONLINESUBSYSTEM_API const FString AuthTypeParam;
-		UE_DEPRECATED(5.2, "AuthPasswordParam is deprecated, please use UE_ONLINE_ERROR_PARAM_IDENTITY_AUTH_PASSWORD instead.")
-		extern ONLINESUBSYSTEM_API const FString AuthPasswordParam;
-
-		// Results
-		UE_DEPRECATED(5.2, "NoUserId is deprecated, please use UE_ONLINE_ERROR_PARAM_IDENTITY_NO_USER_ID instead.")
-		extern ONLINESUBSYSTEM_API const FString NoUserId;
-		UE_DEPRECATED(5.2, "NoAuthToken is deprecated, please use UE_ONLINE_ERROR_PARAM_IDENTITY_NO_AUTH_TOKEN instead.")
-		extern ONLINESUBSYSTEM_API const FString NoAuthToken;
-		UE_DEPRECATED(5.2, "NoAuthType is deprecated, please use UE_ONLINE_ERROR_PARAM_IDENTITY_NO_AUTH_TYPE instead.")
-		extern ONLINESUBSYSTEM_API const FString NoAuthType;
 	}
 }
 
diff --git a/Engine/Plugins/Online/OnlineSubsystem/Source/Public/OnlineSubsystemNames.h b/Engine/Plugins/Online/OnlineSubsystem/Source/Public/OnlineSubsystemNames.h
--- a/Engine/Plugins/Online/OnlineSubsystem/Source/Public/OnlineSubsystemNames.h
+++ b/Engine/Plugins/Online/OnlineSubsystem/Source/Public/OnlineSubsystemNames.h
@@ -85,16 +85,10 @@ extern ONLINESUBSYSTEM_API const FName TENCENT_SUBSYSTEM;
 extern ONLINESUBSYSTEM_API const FName SWITCH_SUBSYSTEM;
 #endif
 
-UE_DEPRECATED(5.2, "OnlineSubsystemOculus has been deprecated and will be removed") ONLINESUBSYSTEM_API extern const FName OCULUS_SUBSYSTEM;
-
 #ifndef SAMSUNG_SUBSYSTEM
 extern ONLINESUBSYSTEM_API const FName SAMSUNG_SUBSYSTEM;
 #endif
 
-#ifndef QUAIL_SUBSYSTEM
-extern ONLINESUBSYSTEM_API const FName QUAIL_SUBSYSTEM;
-#endif
-
 #ifndef EOS_SUBSYSTEM
 extern ONLINESUBSYSTEM_API const FName EOS_SUBSYSTEM;
 #endif
diff --git a/Engine/Plugins/Online/OnlineSubsystemOculus/Config/Localization/OnlineSubsystemOculus.ini b/Engine/Plugins/Online/OnlineSubsystemOculus/Config/Localization/OnlineSubsystemOculus.ini
deleted file mode 100644
diff --git a/Engine/Plugins/Online/OnlineSubsystemOculus/OnlineSubsystemOculus.uplugin b/Engine/Plugins/Online/OnlineSubsystemOculus/OnlineSubsystemOculus.uplugin
deleted file mode 100644
diff --git a/Engine/Plugins/Online/OnlineSubsystemOculus/Source/Classes/OculusCreateSessionCallbackProxy.h b/Engine/Plugins/Online/OnlineSubsystemOculus/Source/Classes/OculusCreateSessionCallbackProxy.h
deleted file mode 100644
diff --git a/Engine/Plugins/Online/OnlineSubsystemOculus/Source/Classes/OculusEntitlementCallbackProxy.h b/Engine/Plugins/Online/OnlineSubsystemOculus/Source/Classes/OculusEntitlementCallbackProxy.h
deleted file mode 100644
diff --git a/Engine/Plugins/Online/OnlineSubsystemOculus/Source/Classes/OculusFindSessionsCallbackProxy.h b/Engine/Plugins/Online/OnlineSubsystemOculus/Source/Classes/OculusFindSessionsCallbackProxy.h
deleted file mode 100644
diff --git a/Engine/Plugins/Online/OnlineSubsystemOculus/Source/Classes/OculusIdentityCallbackProxy.h b/Engine/Plugins/Online/OnlineSubsystemOculus/Source/Classes/OculusIdentityCallbackProxy.h
deleted file mode 100644
diff --git a/Engine/Plugins/Online/OnlineSubsystemOculus/Source/Classes/OculusNetConnection.h b/Engine/Plugins/Online/OnlineSubsystemOculus/Source/Classes/OculusNetConnection.h
deleted file mode 100644
diff --git a/Engine/Plugins/Online/OnlineSubsystemOculus/Source/Classes/OculusNetDriver.h b/Engine/Plugins/Online/OnlineSubsystemOculus/Source/Classes/OculusNetDriver.h
deleted file mode 100644
diff --git a/Engine/Plugins/Online/OnlineSubsystemOculus/Source/Classes/OculusUpdateSessionCallbackProxy.h b/Engine/Plugins/Online/OnlineSubsystemOculus/Source/Classes/OculusUpdateSessionCallbackProxy.h
deleted file mode 100644
diff --git a/Engine/Plugins/Online/OnlineSubsystemOculus/Source/OnlineSubsystemOculus.Build.cs b/Engine/Plugins/Online/OnlineSubsystemOculus/Source/OnlineSubsystemOculus.Build.cs
deleted file mode 100644
diff --git a/Engine/Plugins/Online/OnlineSubsystemOculus/Source/Private/IPAddressOculus.h b/Engine/Plugins/Online/OnlineSubsystemOculus/Source/Private/IPAddressOculus.h
deleted file mode 100644
diff --git a/Engine/Plugins/Online/OnlineSubsystemOculus/Source/Private/OVR_PlatformLoader.cpp b/Engine/Plugins/Online/OnlineSubsystemOculus/Source/Private/OVR_PlatformLoader.cpp
deleted file mode 100644
diff --git a/Engine/Plugins/Online/OnlineSubsystemOculus/Source/Private/OculusCreateSessionCallbackProxy.cpp b/Engine/Plugins/Online/OnlineSubsystemOculus/Source/Private/OculusCreateSessionCallbackProxy.cpp
deleted file mode 100644
diff --git a/Engine/Plugins/Online/OnlineSubsystemOculus/Source/Private/OculusEntitlementCallbackProxy.cpp b/Engine/Plugins/Online/OnlineSubsystemOculus/Source/Private/OculusEntitlementCallbackProxy.cpp
deleted file mode 100644
diff --git a/Engine/Plugins/Online/OnlineSubsystemOculus/Source/Private/OculusFindSessionsCallbackProxy.cpp b/Engine/Plugins/Online/OnlineSubsystemOculus/Source/Private/OculusFindSessionsCallbackProxy.cpp
deleted file mode 100644
diff --git a/Engine/Plugins/Online/OnlineSubsystemOculus/Source/Private/OculusIdentityCallbackProxy.cpp b/Engine/Plugins/Online/OnlineSubsystemOculus/Source/Private/OculusIdentityCallbackProxy.cpp
deleted file mode 100644
diff --git a/Engine/Plugins/Online/OnlineSubsystemOculus/Source/Private/OculusNetConnection.cpp b/Engine/Plugins/Online/OnlineSubsystemOculus/Source/Private/OculusNetConnection.cpp
deleted file mode 100644
diff --git a/Engine/Plugins/Online/OnlineSubsystemOculus/Source/Private/OculusNetDriver.cpp b/Engine/Plugins/Online/OnlineSubsystemOculus/Source/Private/OculusNetDriver.cpp
deleted file mode 100644
diff --git a/Engine/Plugins/Online/OnlineSubsystemOculus/Source/Private/OculusUpdateSessionCallbackProxy.cpp b/Engine/Plugins/Online/OnlineSubsystemOculus/Source/Private/OculusUpdateSessionCallbackProxy.cpp
deleted file mode 100644
diff --git a/Engine/Plugins/Online/OnlineSubsystemOculus/Source/Private/OnlineAchievementsInterfaceOculus.cpp b/Engine/Plugins/Online/OnlineSubsystemOculus/Source/Private/OnlineAchievementsInterfaceOculus.cpp
deleted file mode 100644
diff --git a/Engine/Plugins/Online/OnlineSubsystemOculus/Source/Private/OnlineAchievementsInterfaceOculus.h b/Engine/Plugins/Online/OnlineSubsystemOculus/Source/Private/OnlineAchievementsInterfaceOculus.h
deleted file mode 100644
diff --git a/Engine/Plugins/Online/OnlineSubsystemOculus/Source/Private/OnlineFriendsInterfaceOculus.cpp b/Engine/Plugins/Online/OnlineSubsystemOculus/Source/Private/OnlineFriendsInterfaceOculus.cpp
deleted file mode 100644
diff --git a/Engine/Plugins/Online/OnlineSubsystemOculus/Source/Private/OnlineFriendsInterfaceOculus.h b/Engine/Plugins/Online/OnlineSubsystemOculus/Source/Private/OnlineFriendsInterfaceOculus.h
deleted file mode 100644
diff --git a/Engine/Plugins/Online/OnlineSubsystemOculus/Source/Private/OnlineIdentityOculus.cpp b/Engine/Plugins/Online/OnlineSubsystemOculus/Source/Private/OnlineIdentityOculus.cpp
deleted file mode 100644
diff --git a/Engine/Plugins/Online/OnlineSubsystemOculus/Source/Private/OnlineIdentityOculus.h b/Engine/Plugins/Online/OnlineSubsystemOculus/Source/Private/OnlineIdentityOculus.h
deleted file mode 100644
diff --git a/Engine/Plugins/Online/OnlineSubsystemOculus/Source/Private/OnlineLeaderboardInterfaceOculus.cpp b/Engine/Plugins/Online/OnlineSubsystemOculus/Source/Private/OnlineLeaderboardInterfaceOculus.cpp
deleted file mode 100644
diff --git a/Engine/Plugins/Online/OnlineSubsystemOculus/Source/Private/OnlineLeaderboardInterfaceOculus.h b/Engine/Plugins/Online/OnlineSubsystemOculus/Source/Private/OnlineLeaderboardInterfaceOculus.h
deleted file mode 100644
diff --git a/Engine/Plugins/Online/OnlineSubsystemOculus/Source/Private/OnlineMessageMultiTaskOculus.cpp b/Engine/Plugins/Online/OnlineSubsystemOculus/Source/Private/OnlineMessageMultiTaskOculus.cpp
deleted file mode 100644
diff --git a/Engine/Plugins/Online/OnlineSubsystemOculus/Source/Private/OnlineMessageMultiTaskOculus.h b/Engine/Plugins/Online/OnlineSubsystemOculus/Source/Private/OnlineMessageMultiTaskOculus.h
deleted file mode 100644
diff --git a/Engine/Plugins/Online/OnlineSubsystemOculus/Source/Private/OnlineMessageTaskManagerOculus.cpp b/Engine/Plugins/Online/OnlineSubsystemOculus/Source/Private/OnlineMessageTaskManagerOculus.cpp
deleted file mode 100644
diff --git a/Engine/Plugins/Online/OnlineSubsystemOculus/Source/Private/OnlineSessionInterfaceOculus.cpp b/Engine/Plugins/Online/OnlineSubsystemOculus/Source/Private/OnlineSessionInterfaceOculus.cpp
deleted file mode 100644
diff --git a/Engine/Plugins/Online/OnlineSubsystemOculus/Source/Private/OnlineSessionInterfaceOculus.h b/Engine/Plugins/Online/OnlineSubsystemOculus/Source/Private/OnlineSessionInterfaceOculus.h
deleted file mode 100644
diff --git a/Engine/Plugins/Online/OnlineSubsystemOculus/Source/Private/OnlineSubsystemOculus.cpp b/Engine/Plugins/Online/OnlineSubsystemOculus/Source/Private/OnlineSubsystemOculus.cpp
deleted file mode 100644
diff --git a/Engine/Plugins/Online/OnlineSubsystemOculus/Source/Private/OnlineSubsystemOculusModule.cpp b/Engine/Plugins/Online/OnlineSubsystemOculus/Source/Private/OnlineSubsystemOculusModule.cpp
deleted file mode 100644
diff --git a/Engine/Plugins/Online/OnlineSubsystemOculus/Source/Private/OnlineSubsystemOculusPrivate.h b/Engine/Plugins/Online/OnlineSubsystemOculus/Source/Private/OnlineSubsystemOculusPrivate.h
deleted file mode 100644
diff --git a/Engine/Plugins/Online/OnlineSubsystemOculus/Source/Private/OnlineSubsystemOculusTypes.h b/Engine/Plugins/Online/OnlineSubsystemOculus/Source/Private/OnlineSubsystemOculusTypes.h
deleted file mode 100644
diff --git a/Engine/Plugins/Online/OnlineSubsystemOculus/Source/Private/OnlineUserCloudOculus.cpp b/Engine/Plugins/Online/OnlineSubsystemOculus/Source/Private/OnlineUserCloudOculus.cpp
deleted file mode 100644
diff --git a/Engine/Plugins/Online/OnlineSubsystemOculus/Source/Private/OnlineUserCloudOculus.h b/Engine/Plugins/Online/OnlineSubsystemOculus/Source/Private/OnlineUserCloudOculus.h
deleted file mode 100644
diff --git a/Engine/Plugins/Online/OnlineSubsystemOculus/Source/Private/OnlineVoiceOculus.cpp b/Engine/Plugins/Online/OnlineSubsystemOculus/Source/Private/OnlineVoiceOculus.cpp
deleted file mode 100644
diff --git a/Engine/Plugins/Online/OnlineSubsystemOculus/Source/Private/OnlineVoiceOculus.h b/Engine/Plugins/Online/OnlineSubsystemOculus/Source/Private/OnlineVoiceOculus.h
deleted file mode 100644
diff --git a/Engine/Plugins/Online/OnlineSubsystemOculus/Source/Public/OnlineMessageTaskManagerOculus.h b/Engine/Plugins/Online/OnlineSubsystemOculus/Source/Public/OnlineMessageTaskManagerOculus.h
deleted file mode 100644
diff --git a/Engine/Plugins/Online/OnlineSubsystemOculus/Source/Public/OnlineSubsystemOculus.h b/Engine/Plugins/Online/OnlineSubsystemOculus/Source/Public/OnlineSubsystemOculus.h
deleted file mode 100644
diff --git a/Engine/Plugins/Online/OnlineSubsystemOculus/Source/Public/OnlineSubsystemOculusModule.h b/Engine/Plugins/Online/OnlineSubsystemOculus/Source/Public/OnlineSubsystemOculusModule.h
deleted file mode 100644
diff --git a/Engine/Plugins/Online/OnlineSubsystemOculus/Source/Public/OnlineSubsystemOculusPackage.h b/Engine/Plugins/Online/OnlineSubsystemOculus/Source/Public/OnlineSubsystemOculusPackage.h
deleted file mode 100644
diff --git a/Engine/Plugins/Online/OnlineSubsystemUtils/Source/OnlineSubsystemUtils/Private/OnlineSubsystemUtilsModule.cpp b/Engine/Plugins/Online/OnlineSubsystemUtils/Source/OnlineSubsystemUtils/Private/OnlineSubsystemUtilsModule.cpp
--- a/Engine/Plugins/Online/OnlineSubsystemUtils/Source/OnlineSubsystemUtils/Private/OnlineSubsystemUtilsModule.cpp
+++ b/Engine/Plugins/Online/OnlineSubsystemUtils/Source/OnlineSubsystemUtils/Private/OnlineSubsystemUtilsModule.cpp
@@ -238,11 +238,6 @@ private:
 		CREATE_HASH(AMAZON_SUBSYSTEM);
 		CREATE_HASH(GAMECIRCLE_SUBSYSTEM);
 		CREATE_HASH(THUNDERHEAD_SUBSYSTEM); 
-PRAGMA_DISABLE_DEPRECATION_WARNINGS
-		CREATE_HASH(OCULUS_SUBSYSTEM); 
-PRAGMA_ENABLE_DEPRECATION_WARNINGS
-
-		CREATE_HASH(QUAIL_SUBSYSTEM);
 		// Shouldn't need these as they are mocking interfaces for existing platforms
 		CREATE_HASH(PS4SERVER_SUBSYSTEM);
 
diff --git a/Engine/Plugins/Online/OnlineSubsystemUtils/Source/OnlineSubsystemUtils/Private/Tests/OnlineSubsystemUtilitiesTests.cpp b/Engine/Plugins/Online/OnlineSubsystemUtils/Source/OnlineSubsystemUtils/Private/Tests/OnlineSubsystemUtilitiesTests.cpp
--- a/Engine/Plugins/Online/OnlineSubsystemUtils/Source/OnlineSubsystemUtils/Private/Tests/OnlineSubsystemUtilitiesTests.cpp
+++ b/Engine/Plugins/Online/OnlineSubsystemUtils/Source/OnlineSubsystemUtils/Private/Tests/OnlineSubsystemUtilitiesTests.cpp
@@ -55,7 +55,6 @@ namespace OSSUtilsTestHelper
 		OSSList.AddUnique(TEXT( "Live" ) );
 		OSSList.AddUnique(TEXT( "Thunderhead" ) );
 		OSSList.AddUnique(TEXT( "Amazon" ) );
-		OSSList.AddUnique(TEXT( "Oculus" ) );
 
 		for ( auto& OSSName : OSSList )
 		{
