[AutoRTFM] Migrate to UE::TScopeLock instead of FTransactionallySafeScopeLock.

These are functionally identical.

diff --git a/Engine/Plugins/Online/OnlineSubsystem/Source/Private/OnlineSubsystemModule.cpp b/Engine/Plugins/Online/OnlineSubsystem/Source/Private/OnlineSubsystemModule.cpp
--- a/Engine/Plugins/Online/OnlineSubsystem/Source/Private/OnlineSubsystemModule.cpp
+++ b/Engine/Plugins/Online/OnlineSubsystem/Source/Private/OnlineSubsystemModule.cpp
@@ -2,7 +2,7 @@
 
 #include "OnlineSubsystemModule.h"
 #include "Misc/ConfigCacheIni.h"
-#include "Misc/TransactionallySafeScopeLock.h"
+#include "Misc/ScopeLock.h"
 #include "OnlineSubsystemImpl.h"
 #include "OnlineDelegates.h"
 
@@ -228,7 +228,7 @@ void FOnlineSubsystemModule::ReloadDefaultSubsystem()
 void FOnlineSubsystemModule::PreUnloadOnlineSubsystem()
 {
 	// Shutdown all online subsystem instances
-	FTransactionallySafeScopeLock Lock(&OnlineSubsystemsLock);
+	UE::TScopeLock Lock(OnlineSubsystemsLock);
 	for (TMap<FName, IOnlineSubsystemPtr>::TIterator It(OnlineSubsystems); It; ++It)
 	{
 		It.Value()->PreUnload();
@@ -241,7 +241,7 @@ void FOnlineSubsystemModule::ShutdownOnlineSubsystem()
 
 	// Shutdown all online subsystem instances
 	{
-		FTransactionallySafeScopeLock Lock(&OnlineSubsystemsLock);
+		UE::TScopeLock Lock(OnlineSubsystemsLock);
 		for (TMap<FName, IOnlineSubsystemPtr>::TIterator It(OnlineSubsystems); It; ++It)
 		{
 			It.Value()->Shutdown();
@@ -281,7 +281,7 @@ void FOnlineSubsystemModule::UnregisterPlatformService(const FName FactoryName)
 
 void FOnlineSubsystemModule::EnumerateOnlineSubsystems(FEnumerateOnlineSubsystemCb& EnumCb)
 {
-	FTransactionallySafeScopeLock Lock(&OnlineSubsystemsLock);
+	UE::TScopeLock Lock(OnlineSubsystemsLock);
 	for (TPair<FName, IOnlineSubsystemPtr>& OnlineSubsystem : OnlineSubsystems)
 	{
 		if (OnlineSubsystem.Value.IsValid())
@@ -374,7 +374,7 @@ IOnlineSubsystem* FOnlineSubsystemModule::GetOnlineSubsystem(const FName InSubsy
 	{
 		bool bWasNewlyCreated = false;
 		{
-			FTransactionallySafeScopeLock Lock(&OnlineSubsystemsLock);
+			UE::TScopeLock Lock(OnlineSubsystemsLock);
 			OnlineSubsystem = FindExistingSubsystem(OnlineSubsystems, KeyName);
 			if (!OnlineSubsystem)
 			{
@@ -449,7 +449,7 @@ void FOnlineSubsystemModule::DestroyOnlineSubsystem(const FName InSubsystemName)
 	{
 		IOnlineSubsystemPtr OnlineSubsystem;
 		{
-			FTransactionallySafeScopeLock Lock(&OnlineSubsystemsLock);
+			UE::TScopeLock Lock(OnlineSubsystemsLock);
 			OnlineSubsystems.RemoveAndCopyValue(KeyName, OnlineSubsystem);
 		}
 		if (OnlineSubsystem.IsValid())
@@ -474,7 +474,7 @@ bool FOnlineSubsystemModule::DoesInstanceExist(const FName InSubsystemName) cons
 	{
 		IOnlineSubsystemPtr OnlineSubsystem;
 		{
-			FTransactionallySafeScopeLock Lock(&OnlineSubsystemsLock);
+			UE::TScopeLock Lock(OnlineSubsystemsLock);
 			OnlineSubsystem = FindExistingSubsystem(OnlineSubsystems, KeyName);
 		}
 		return OnlineSubsystem.IsValid();
