- Add warning when OSS loads modules _after_ StartupModule.
-- Loading modules after StartupModule causes issues when those modules shut down _before_ OSS ShutdownModule. The quick fix is to add the modules to AdditionalModulesToLoad.

#rnx

diff --git a/Engine/Plugins/Online/OnlineSubsystem/Source/Private/OnlineSubsystemModule.cpp b/Engine/Plugins/Online/OnlineSubsystem/Source/Private/OnlineSubsystemModule.cpp
--- a/Engine/Plugins/Online/OnlineSubsystem/Source/Private/OnlineSubsystemModule.cpp
+++ b/Engine/Plugins/Online/OnlineSubsystem/Source/Private/OnlineSubsystemModule.cpp
@@ -8,6 +8,11 @@
 
 IMPLEMENT_MODULE( FOnlineSubsystemModule, OnlineSubsystem );
 
+namespace
+{
+	static bool bPostStartupModule = false;
+}
+
 /**
  * OnlineDelegates.h static declarations
  */
@@ -54,6 +59,7 @@ static IModuleInterface* LoadSubsystemModule(const FString& SubsystemName, const
 
 		if (!ModuleManager.IsModuleLoaded(ModuleName))
 		{
+			UE_CLOG(bPostStartupModule, LogOnline, Warning, TEXT("%hs attempting to load module \"%s\" after FOnlineSubsystemModule::StartupModule. This can result in shutdown issues due to the loaded module unloading before OSS itself. Please add the module to the [OnlineSubsystem] AdditionalModulesToLoad config array to fix."), __FUNCTION__, *ModuleName.ToString());
 			// Attempt to load the module
 			ModuleManager.LoadModule(ModuleName);
 		}
@@ -103,6 +109,8 @@ void FOnlineSubsystemModule::StartupModule()
 		ProcessConfigDefinedSubsystems();
 		IOnlineSubsystem::GetByPlatform();
 	}
+
+	bPostStartupModule = true;
 }
 
 void FOnlineSubsystemModule::PreUnloadCallback()
