Set/modify the BuildIDOverride from inside the client

diff --git a/Engine/Plugins/Online/OnlineBase/Source/Public/Online/OnlineBase.cpp b/Engine/Plugins/Online/OnlineBase/Source/Public/Online/OnlineBase.cpp
new file mode 100644
--- /dev/null
+++ b/Engine/Plugins/Online/OnlineBase/Source/Public/Online/OnlineBase.cpp
@@ -0,0 +1,11 @@
+﻿#include "OnlineBase.h"
+
+static TAutoConsoleVariable<int32> CVarBuildIdOverride(
+TEXT("buildidoverride2"),
+0,
+TEXT("Sets build id used for matchmaking "));
+
+TAutoConsoleVariable<int32>& GetBuildIdOverrideCVar()
+{
+	return CVarBuildIdOverride;
+}
diff --git a/Engine/Plugins/Online/OnlineBase/Source/Public/Online/OnlineBase.h b/Engine/Plugins/Online/OnlineBase/Source/Public/Online/OnlineBase.h
--- a/Engine/Plugins/Online/OnlineBase/Source/Public/Online/OnlineBase.h
+++ b/Engine/Plugins/Online/OnlineBase/Source/Public/Online/OnlineBase.h
@@ -2,6 +2,9 @@
 
 #pragma once
 
+#include "HAL/Platform.h"
+#include "HAL/IConsoleManager.h"
+
 #ifndef ONLINE_SUCCESS
 #define ONLINE_SUCCESS 0
 #endif
@@ -12,4 +15,6 @@
 
 #ifndef ONLINE_IO_PENDING
 #define ONLINE_IO_PENDING 997
-#endif
\ No newline at end of file
+#endif
+
+ONLINEBASE_API TAutoConsoleVariable<int32>&  GetBuildIdOverrideCVar();
\ No newline at end of file
diff --git a/Engine/Plugins/Online/OnlineServices/Source/OnlineServicesInterface/Private/Online/OnlineServices.cpp b/Engine/Plugins/Online/OnlineServices/Source/OnlineServicesInterface/Private/Online/OnlineServices.cpp
--- a/Engine/Plugins/Online/OnlineServices/Source/OnlineServicesInterface/Private/Online/OnlineServices.cpp
+++ b/Engine/Plugins/Online/OnlineServices/Source/OnlineServicesInterface/Private/Online/OnlineServices.cpp
@@ -1,6 +1,7 @@
 // Copyright Epic Games, Inc. All Rights Reserved.
 
 #include "Online/OnlineServices.h"
+#include "Online/OnlineBase.h"
 
 #include "Online/OnlineServicesLog.h"
 #include "Online/OnlineServicesRegistry.h"
@@ -19,12 +20,6 @@ int32 GetBuildUniqueId()
 	static bool bUseBuildIdOverride = false;
 	static int32 BuildIdOverride = 0;
 
-	// add a cvar so it can be modified at runtime
-	static FAutoConsoleVariableRef CVarBuildIdOverride(
-		TEXT("buildidoverride"), BuildId,
-		TEXT("Sets build id used for matchmaking "),
-		ECVF_Default);
-
 	if (!bStaticCheck)
 	{
 		bStaticCheck = true;
@@ -60,9 +55,13 @@ int32 GetBuildUniqueId()
 		{
 			BuildId = BuildIdOverride;
 		}
+		
+		// use a cvar so it can be modified at runtime
+		TAutoConsoleVariable<int32>& CVarBuildIdOverride = GetBuildIdOverrideCVar();
+		CVarBuildIdOverride->Set(BuildId);
 	}
 
-	return BuildId;
+	return GetBuildIdOverrideCVar()->GetInt();
 }
 
 bool IsLoaded(EOnlineServices OnlineServices, FName InstanceName, FName InstanceConfigName)
diff --git a/Engine/Plugins/Online/OnlineSubsystem/Source/Private/OnlineSubsystem.cpp b/Engine/Plugins/Online/OnlineSubsystem/Source/Private/OnlineSubsystem.cpp
--- a/Engine/Plugins/Online/OnlineSubsystem/Source/Private/OnlineSubsystem.cpp
+++ b/Engine/Plugins/Online/OnlineSubsystem/Source/Private/OnlineSubsystem.cpp
@@ -8,6 +8,7 @@
 #include "Online/OnlineSessionNames.h"
 #include "OnlineIdentityErrors.h" // IWYU pragma: keep
 #include "OnlineSessionSettings.h"
+#include "Online/OnlineBase.h"
 
 #include "Interfaces/OnlineChatInterface.h"
 #include "Interfaces/OnlineIdentityInterface.h"
@@ -210,12 +211,7 @@ int32 GetBuildUniqueId()
 	static int32 BuildId = 0;
 	static bool bUseBuildIdOverride = false;
 	static int32 BuildIdOverride = 0;
-
-	// add a cvar so it can be modified at runtime
-	static FAutoConsoleVariableRef CVarBuildIdOverride(
-		TEXT("buildidoverride"), BuildId,
-		TEXT("Sets build id used for matchmaking "),
-		ECVF_Default);
+	
 
 	if (!bStaticCheck)
 	{
@@ -253,12 +249,17 @@ int32 GetBuildUniqueId()
 			BuildId = BuildIdOverride;
 		}
 
+		// use a cvar so it can be modified at runtime
+		TAutoConsoleVariable<int32>& CVarBuildIdOverride = GetBuildIdOverrideCVar();
+		CVarBuildIdOverride->Set(BuildId);
+		
+
 		// Will add an info entry to the console
 		static FTrackedActivity Ta(TEXT("BuildId"), *FString::FromInt(BuildId), FTrackedActivity::ELight::None, FTrackedActivity::EType::Info);
 		CVarBuildIdOverride->SetOnChangedCallback(FConsoleVariableDelegate::CreateLambda([](IConsoleVariable* CVar) { Ta.Update(*CVar->GetString()); }));
 	}
 
-	return BuildId;
+	return GetBuildIdOverrideCVar()->GetInt();
 }
 
 TAutoConsoleVariable<FString> CVarPlatformOverride(
