Add the support to cancel timeout online async task;
Output warning log when breached task eventually completed.

diff --git a/Engine/Plugins/Online/OnlineSubsystem/Source/Private/OnlineAsyncTaskManager.cpp b/Engine/Plugins/Online/OnlineSubsystem/Source/Private/OnlineAsyncTaskManager.cpp
--- a/Engine/Plugins/Online/OnlineSubsystem/Source/Private/OnlineAsyncTaskManager.cpp
+++ b/Engine/Plugins/Online/OnlineSubsystem/Source/Private/OnlineAsyncTaskManager.cpp
@@ -348,6 +348,7 @@ void FOnlineAsyncTaskManager::Tick()
 			}
 
 			Task->Tick();
+			Task->CancelWhenTimeout();
 
 			if (Task->IsDone())
 			{
@@ -392,6 +393,7 @@ void FOnlineAsyncTaskManager::Tick()
 				OnlineThreadTickInfo.CurrentTask = Task;
 			}
 			Task->Tick();
+			Task->CancelWhenTimeout();
 
 			if (Task->IsDone())
 			{
@@ -411,6 +413,11 @@ void FOnlineAsyncTaskManager::Tick()
 				// Task is done, add to the outgoing queue
 				AddToOutQueue(Task);
 
+				if (bActiveSerialTaskBreachReported)
+				{
+					UE_LOG_ONLINE(Warning, TEXT("The breached serial task eventually completed, Task=[%s]"), *Task->ToString());
+				}
+
 				{
 					FScopeLock Lock(&ActiveSerialTaskLock);
 					ActiveSerialTask = nullptr;
diff --git a/Engine/Plugins/Online/OnlineSubsystem/Source/Public/OnlineAsyncTaskManager.h b/Engine/Plugins/Online/OnlineSubsystem/Source/Public/OnlineAsyncTaskManager.h
--- a/Engine/Plugins/Online/OnlineSubsystem/Source/Public/OnlineAsyncTaskManager.h
+++ b/Engine/Plugins/Online/OnlineSubsystem/Source/Public/OnlineAsyncTaskManager.h
@@ -138,6 +138,12 @@ public:
 		// assert that we're not on the game thread
 		check(!IsInGameThread() || !FPlatformProcess::SupportsMultithreading());
 	}
+
+	/**
+	 * If cancel is supported, cancel when timeout if timeout is set for this task
+	 * Not putting this in Tick directly because existing subclasses could be overriding Tick without calling super::Tick
+	 */
+	virtual void CancelWhenTimeout() {}
 };
 
 
@@ -310,7 +316,7 @@ private:
 	/** When the active serial task was started on the game thread */
 	double ActiveSerialTaskStartTime = 0.0;
 	/** Has a breach been reported for the active task */
-	bool bActiveSerialTaskBreachReported = false;
+	std::atomic<bool> bActiveSerialTaskBreachReported = false;
 
 	struct FOnlineThreadTickInfo
 	{
