Allow "ENABLE_PGO_PROFILE 1" Shipping configuration to disable online subsystems from the commandline

#tests preflight PGO

diff --git a/Engine/Plugins/Online/OnlineSubsystem/Source/Private/OnlineSubsystem.cpp b/Engine/Plugins/Online/OnlineSubsystem/Source/Private/OnlineSubsystem.cpp
--- a/Engine/Plugins/Online/OnlineSubsystem/Source/Private/OnlineSubsystem.cpp
+++ b/Engine/Plugins/Online/OnlineSubsystem/Source/Private/OnlineSubsystem.cpp
@@ -377,8 +377,8 @@ FAutoConsoleCommand CmdResetAchievements(
 bool IOnlineSubsystem::IsEnabled(const FName& SubsystemName, const FName& InstanceName)
 {
 	bool bIsDisabledByCommandLine = false;
-#if !UE_BUILD_SHIPPING && !UE_BUILD_SHIPPING_WITH_EDITOR
-	// In non-shipping builds, check for a command line override to disable the OSS
+#if (!UE_BUILD_SHIPPING && !UE_BUILD_SHIPPING_WITH_EDITOR) || ENABLE_PGO_PROFILE
+	// In non-shipping builds, or PGO_PROFILE enabled Shipping builds, check for a command line override to disable the OSS
 	bIsDisabledByCommandLine = FParse::Param(FCommandLine::Get(), *FString::Printf(TEXT("no%s"), *SubsystemName.ToString()));
 #endif
 	
