Removal of online code marked for deprecation in 5.5

diff --git a/Engine/Plugins/Online/OnlineBase/Source/Public/Online/OnlineSessionNames.h b/Engine/Plugins/Online/OnlineBase/Source/Public/Online/OnlineSessionNames.h
--- a/Engine/Plugins/Online/OnlineBase/Source/Public/Online/OnlineSessionNames.h
+++ b/Engine/Plugins/Online/OnlineBase/Source/Public/Online/OnlineSessionNames.h
@@ -131,10 +131,6 @@
 /** Search for secure servers only (value is true/false) */
 #define SEARCH_SECURE_SERVERS_ONLY FName(TEXT("SECUREONLY"))
 
-/** Search for presence sessions only (value is true/false) */
-UE_DEPRECATED(5.5, "SEARCH_PRESENCE (\"PRESENCESEARCH\") is deprecated and will soon stop being a valid UE-defined key. Please consult upgrade notes for more details")
-const FName SEARCH_PRESENCE = FName(TEXT("PRESENCESEARCH"));
-
 /** Search for a match with min player availability (value is int) */
 #define SEARCH_MINSLOTSAVAILABLE FName(TEXT("MINSLOTSAVAILABLE"))
 /** Exclude all matches where any unique ids in a given array are present (value is string of the form "uniqueid1;uniqueid2;uniqueid3") */
diff --git a/Engine/Plugins/Online/OnlineFramework/Source/Party/Private/Party/SocialParty.cpp b/Engine/Plugins/Online/OnlineFramework/Source/Party/Private/Party/SocialParty.cpp
--- a/Engine/Plugins/Online/OnlineFramework/Source/Party/Private/Party/SocialParty.cpp
+++ b/Engine/Plugins/Online/OnlineFramework/Source/Party/Private/Party/SocialParty.cpp
@@ -206,13 +206,6 @@ FPartyJoinApproval USocialParty::EvaluateJoinRequest(const TArray<IOnlinePartyUs
 	return JoinApproval;
 }
 
-PRAGMA_DISABLE_DEPRECATION_WARNINGS
-bool USocialParty::ShouldCacheForRejoinOnDisconnect() const
-{
-	return bEnableAutomaticPartyRejoin && GetNumPartyMembers() > 1;
-}
-PRAGMA_ENABLE_DEPRECATION_WARNINGS
-
 bool USocialParty::IsCurrentlyLeaving() const
 {
 	return bIsLeavingParty;
@@ -1399,20 +1392,6 @@ bool USocialParty::ContainsUser(const USocialUser& User) const
 	return false;
 }
 
-ULocalPlayer& USocialParty::GetOwningLocalPlayer() const
-{
-	// This is deprecated as it was unsafe.
-	// The Toolkit's LocalPlayerOwner is a TWeakObjectPtr and may return nullptr when the local player logs out. Please use the pointer version.
-	if (ULocalPlayer* LocalPlayer = GetOwningLocalPlayerPtr())
-	{
-		return *LocalPlayer;
-	}
-	else
-	{
-		return *NewObject<ULocalPlayer>();
-	}
-}
-
 ULocalPlayer* USocialParty::GetOwningLocalPlayerPtr() const
 {
 	//@todo DanH Party: This is a wee bit heavy - should be able to do this in fewer steps
diff --git a/Engine/Plugins/Online/OnlineFramework/Source/Party/Private/SocialManager.cpp b/Engine/Plugins/Online/OnlineFramework/Source/Party/Private/SocialManager.cpp
--- a/Engine/Plugins/Online/OnlineFramework/Source/Party/Private/SocialManager.cpp
+++ b/Engine/Plugins/Online/OnlineFramework/Source/Party/Private/SocialManager.cpp
@@ -29,31 +29,6 @@ static TAutoConsoleVariable<bool> CVarForceDisconnectedToPartyService(
 	false,
 	TEXT("Overrides the bIsConnectedToPartyService state"));
 
-
-//////////////////////////////////////////////////////////////////////////
-// FRejoinableParty
-//////////////////////////////////////////////////////////////////////////
-
-PRAGMA_DISABLE_DEPRECATION_WARNINGS
-USocialManager::FRejoinableParty::FRejoinableParty(const USocialParty& SourceParty)
-	: PartyId(SourceParty.GetPartyId().AsShared()),
-	OriginalJoinMethod(SourceParty.GetOwningLocalMember().GetRepData().GetJoinMethod())
-{
-	for (UPartyMember* Member : SourceParty.GetPartyMembers())
-	{
-		// Rejoin fails if the local user's Id is in there as a member, so only include everyone else
-		if (Member != &SourceParty.GetOwningLocalMember())
-		{
-			FUniqueNetIdRepl MemberId = Member->GetPrimaryNetId();
-			if (MemberId.IsValid())
-			{
-				MemberIds.Add(MemberId->AsShared());
-			}
-		}		
-	}
-}
-PRAGMA_ENABLE_DEPRECATION_WARNINGS
-
 //////////////////////////////////////////////////////////////////////////
 // FJoinPartyAttempt
 //////////////////////////////////////////////////////////////////////////
@@ -65,25 +40,14 @@ USocialManager::FJoinPartyAttempt::FJoinPartyAttempt(const USocialUser* InTarget
 	, OnJoinComplete(InOnJoinComplete)
 {}
 
-PRAGMA_DISABLE_DEPRECATION_WARNINGS
-USocialManager::FJoinPartyAttempt::FJoinPartyAttempt(TSharedRef<const FRejoinableParty> InRejoinInfo)
-	: PartyTypeId(IOnlinePartySystem::GetPrimaryPartyTypeId())
-	, JoinMethod(InRejoinInfo->OriginalJoinMethod)
-	, RejoinInfo(InRejoinInfo)
-{}
-PRAGMA_ENABLE_DEPRECATION_WARNINGS
-
 FString USocialManager::FJoinPartyAttempt::ToDebugString() const
 {
-PRAGMA_DISABLE_DEPRECATION_WARNINGS
-	return FString::Printf(TEXT("IsRejoin (%s), TargetUser (%s), PartyId (%s), TypeId (%d), TargetUserPlatformId (%s), JoinMethod (%s)"),
-		RejoinInfo.IsValid() ? TEXT("true") : TEXT("false"),
+	return FString::Printf(TEXT("TargetUser(%s), PartyId(%s), TypeId(%d), TargetUserPlatformId(%s), JoinMethod(%s)"),
 		TargetUser.IsValid() ? *TargetUser->ToDebugString() : TEXT("invalid"),
-		JoinInfo.IsValid() ? *JoinInfo->GetPartyId()->ToDebugString() : RejoinInfo.IsValid() ? *RejoinInfo->PartyId->ToDebugString() : TEXT("unknown"),
+		JoinInfo.IsValid() ? *JoinInfo->GetPartyId()->ToDebugString() : TEXT("unknown"),
 		PartyTypeId.GetValue(),
 		*TargetUserPlatformId.ToDebugString(),
 		*JoinMethod.ToString());
-PRAGMA_ENABLE_DEPRECATION_WARNINGS
 }
 
 const FName USocialManager::FJoinPartyAttempt::Step_FindPlatformSession = TEXT("FindPlatformSession");
@@ -255,10 +219,6 @@ void USocialManager::ShutdownSocialManager()
 	ShutdownPartiesFunc(JoinedPartiesByTypeId);
 	ShutdownPartiesFunc(LeavingPartiesByTypeId);
 
-PRAGMA_DISABLE_DEPRECATION_WARNINGS
-	RejoinableParty.Reset();
-PRAGMA_ENABLE_DEPRECATION_WARNINGS
-
 	if (SocialDebugTools)
 	{
 		SocialDebugTools->Shutdown();
@@ -676,14 +636,6 @@ ECrossplayPreference USocialManager::GetCrossplayPreference() const
 	return ECrossplayPreference::NoSelection;
 }
 
-PRAGMA_DISABLE_DEPRECATION_WARNINGS
-bool USocialManager::ShouldTryRejoiningPersistentParty(const FRejoinableParty& InRejoinableParty) const
-{
-	// Deprecated
- 	return false;
-}
-PRAGMA_ENABLE_DEPRECATION_WARNINGS
-
 void USocialManager::RefreshCanCreatePartyObjects()
 {
 	const bool bCanNowCreate = CanCreateNewPartyObjects();
@@ -802,19 +754,7 @@ void USocialManager::JoinPartyInternal(FJoinPartyAttempt& JoinAttempt)
 	{
 		JoinAttempt.ActionTimeTracker.BeginStep(FJoinPartyAttempt::Step_JoinParty);
 
-PRAGMA_DISABLE_DEPRECATION_WARNINGS
-		if (JoinAttempt.RejoinInfo.IsValid())
-		{
-			UE_LOG(LogParty, Verbose, TEXT("Attempting to rejoin party [%s] now."), *JoinAttempt.RejoinInfo->PartyId->ToDebugString());
-
-			// Rejoin attempts are initiated differently, but the handler/follow-up is identical to a normal join
-			PartyInterface->RejoinParty(*LocalUserId, *JoinAttempt.RejoinInfo->PartyId, IOnlinePartySystem::GetPrimaryPartyTypeId(), JoinAttempt.RejoinInfo->MemberIds, FOnJoinPartyComplete::CreateUObject(this, &USocialManager::HandleJoinPartyComplete, IOnlinePartySystem::GetPrimaryPartyTypeId()));
-		}
-		else
-PRAGMA_ENABLE_DEPRECATION_WARNINGS
-		{
-			PartyInterface->JoinParty(*LocalUserId, *JoinAttempt.JoinInfo, FOnJoinPartyComplete::CreateUObject(this, &USocialManager::HandleJoinPartyComplete, JoinAttempt.JoinInfo->GetPartyTypeId()));
-		}
+		PartyInterface->JoinParty(*LocalUserId, *JoinAttempt.JoinInfo, FOnJoinPartyComplete::CreateUObject(this, &USocialManager::HandleJoinPartyComplete, JoinAttempt.JoinInfo->GetPartyTypeId()));
 	}
 	else
 	{
@@ -1212,38 +1152,10 @@ void USocialManager::HandlePersistentPartyStateChanged(EPartyState NewState, EPa
 				PersistentParty->LeaveParty();
 			}
 		}
-	
-PRAGMA_DISABLE_DEPRECATION_WARNINGS
-		// If we have other members in our party, then we will try to rejoin this when we come back online
-		if (!RejoinableParty.IsValid() && PersistentParty->ShouldCacheForRejoinOnDisconnect())
-		{
-			UE_LOG(LogParty, Log, TEXT("Caching persistent party [%s] for rejoin"), *PersistentParty->GetPartyId().ToDebugString());
-			RejoinableParty = MakeShared<FRejoinableParty>(*PersistentParty);
-		}
-PRAGMA_ENABLE_DEPRECATION_WARNINGS
 	}
 	else if (NewState == EPartyState::Active)
 	{
 		bIsConnectedToPartyService = true;
-
-PRAGMA_DISABLE_DEPRECATION_WARNINGS
-		if (RejoinableParty.IsValid())
-		{
-			if (ShouldTryRejoiningPersistentParty(*RejoinableParty))
-			{
-				// Bail on the current party, we'll try to rejoin once we've left
-				UE_LOG(LogParty, Log, TEXT("Leaving current persistent party [%s] to attempt to rejoin previous party [%s]"), *PersistentParty->ToDebugString(), *RejoinableParty->PartyId->ToDebugString());
-
-				FJoinPartyAttempt& RejoinAttempt = JoinAttemptsByTypeId.Add(IOnlinePartySystem::GetPrimaryPartyTypeId(), FJoinPartyAttempt(RejoinableParty.ToSharedRef()));
-				RejoinAttempt.ActionTimeTracker.BeginStep(FJoinPartyAttempt::Step_LeaveCurrentParty);
-
-				PersistentParty->LeaveParty();
-			}
-
-			// This is the only time we would try to rejoin, and it's saved on the join attempt if initiated
-			RejoinableParty.Reset();
-		}
-PRAGMA_ENABLE_DEPRECATION_WARNINGS
 	}
 }
 
@@ -1309,9 +1221,7 @@ void USocialManager::HandlePartyLeft(EMemberExitedReason Reason, USocialParty* L
 			JoinAttempt->ActionTimeTracker.CompleteStep(FJoinPartyAttempt::Step_LeaveCurrentParty);
 
 			// We're in the process of joining another party of the same type - do we know where we're heading yet?
-PRAGMA_DISABLE_DEPRECATION_WARNINGS
-			if (JoinAttempt->JoinInfo.IsValid() || JoinAttempt->RejoinInfo.IsValid())
-PRAGMA_ENABLE_DEPRECATION_WARNINGS
+			if (JoinAttempt->JoinInfo.IsValid())
 			{
 				// Join the new party immediately and early out
 				JoinPartyInternal(*JoinAttempt);
diff --git a/Engine/Plugins/Online/OnlineFramework/Source/Party/Private/SocialToolkit.cpp b/Engine/Plugins/Online/OnlineFramework/Source/Party/Private/SocialToolkit.cpp
--- a/Engine/Plugins/Online/OnlineFramework/Source/Party/Private/SocialToolkit.cpp
+++ b/Engine/Plugins/Online/OnlineFramework/Source/Party/Private/SocialToolkit.cpp
@@ -327,20 +327,6 @@ USocialManager& USocialToolkit::GetSocialManager() const
 	return *OuterSocialManager;
 }
 
-ULocalPlayer& USocialToolkit::GetOwningLocalPlayer() const
-{
-	// This is deprecated as it is unsafe.
-	// LocalPlayerOwner is a TWeakObjectPtr and may return nullptr when the local player logs out. Please use the pointer version.
-	if (ULocalPlayer* LocalPlayer = GetOwningLocalPlayerPtr())
-	{
-		return *LocalPlayer;
-	}
-	else
-	{
-		return *NewObject<ULocalPlayer>();
-	}
-}
-
 ULocalPlayer* USocialToolkit::GetOwningLocalPlayerPtr() const
 {
 	return LocalPlayerOwner.Get();
diff --git a/Engine/Plugins/Online/OnlineFramework/Source/Party/Public/Party/PartyMember.h b/Engine/Plugins/Online/OnlineFramework/Source/Party/Public/Party/PartyMember.h
--- a/Engine/Plugins/Online/OnlineFramework/Source/Party/Public/Party/PartyMember.h
+++ b/Engine/Plugins/Online/OnlineFramework/Source/Party/Public/Party/PartyMember.h
@@ -179,17 +179,9 @@ public:
 
 	UE_API virtual void BeginDestroy() override;
 
-	UE_DEPRECATED(5.5, "This has been deprecated to support multiple local players and now requires a performing player.")
-	bool CanPromoteToLeader() const { return false; }
 	UE_API bool CanPromoteToLeader(const ULocalPlayer& PerformingPlayer) const;
-	UE_DEPRECATED(5.5, "This has been deprecated to support multiple local players and now requires a performing player.")
-	bool PromoteToPartyLeader() { return false; }
 	UE_API bool PromoteToPartyLeader(const ULocalPlayer& PerformingPlayer);
-	UE_DEPRECATED(5.5, "This has been deprecated to support multiple local players and now requires a performing player.")
-	bool CanKickFromParty() const { return false; }
 	UE_API bool CanKickFromParty(const ULocalPlayer& PerformingPlayer) const;
-	UE_DEPRECATED(5.5, "This has been deprecated to support multiple local players and now requires a performing player.")
-	bool KickFromParty() { return false; }
 	UE_API bool KickFromParty(const ULocalPlayer& PerformingPlayer);
 
 	UE_API bool IsInitialized() const;
@@ -226,11 +218,6 @@ public:
 	UE_API FString ToDebugString(bool bIncludePartyId = true) const;
 
 protected:
-	UE_DEPRECATED(5.5, "Use InitializePartyMember with an r-value delegate")
-	void InitializePartyMember(const FOnlinePartyMemberConstRef& OssMember, const FSimpleDelegate& OnInitComplete)
-	{
-		InitializePartyMember(OssMember, CopyTemp(OnInitComplete));
-	}
 	UE_API void InitializePartyMember(const FOnlinePartyMemberConstRef& OssMember, FSimpleDelegate&& OnInitComplete);
 
 	FPartyMemberRepData& GetMutableRepData() { return *MemberDataReplicator; }
diff --git a/Engine/Plugins/Online/OnlineFramework/Source/Party/Public/Party/SocialParty.h b/Engine/Plugins/Online/OnlineFramework/Source/Party/Public/Party/SocialParty.h
--- a/Engine/Plugins/Online/OnlineFramework/Source/Party/Public/Party/SocialParty.h
+++ b/Engine/Plugins/Online/OnlineFramework/Source/Party/Public/Party/SocialParty.h
@@ -131,9 +131,6 @@ public:
 
 	UE_API bool ContainsUser(const USocialUser& User) const;
 
-	UE_DEPRECATED(5.5, "GetOwningLocalPlayer returns the Toolkit's LocalPlayerOwner which is a TWeakObjectPtr and may return nullptr when the local player logs out. Please use the pointer version.")
-	UE_API ULocalPlayer& GetOwningLocalPlayer() const;
-
 	UE_API ULocalPlayer* GetOwningLocalPlayerPtr() const;
 	const FUniqueNetIdRepl& GetOwningLocalUserId() const { return OwningLocalUserId; }
 	const FUniqueNetIdRepl& GetPartyLeaderId() const { return CurrentLeaderId; }
@@ -232,9 +229,6 @@ protected:
 	UE_API bool IsInitialized() const;
 	UE_API void TryFinishInitialization();
 
-	UE_DEPRECATED(5.5, "Rejoinable Parties is deprecated and will be removed")
-	UE_API bool ShouldCacheForRejoinOnDisconnect() const;
-
 	UE_API void SetIsMissingPlatformSession(bool bInIsMissingPlatformSession);
 	bool IsMissingPlatformSession() { return bIsMissingPlatformSession; }
 
@@ -244,17 +238,9 @@ protected:
 	// User/member-specific actions that are best exposed on the individuals themselves, but best handled by the actual party
 	UE_API bool HasUserBeenInvited(const USocialUser& User) const;
 
-	UE_DEPRECATED(5.5, "This has been deprecated to support multiple local players and now requires a performing player.")
-	bool CanPromoteMember(const UPartyMember& PartyMember) const { return false; }
 	UE_API bool CanPromoteMember(const ULocalPlayer& PerformingPlayer, const UPartyMember& PartyMember) const;
-	UE_DEPRECATED(5.5, "This has been deprecated to support multiple local players and now requires a performing player.")
-	virtual bool CanKickMember(const UPartyMember& PartyMember) const { return false; }
 	UE_API virtual bool CanKickMember(const ULocalPlayer& PerformingPlayer, const UPartyMember& PartyMember) const;
-	UE_DEPRECATED(5.5, "This has been deprecated to support multiple local players and now requires a performing player.")
-	bool TryPromoteMember(const UPartyMember& PartyMember) { return false; }
 	UE_API bool TryPromoteMember(const ULocalPlayer& PerformingPlayer, const UPartyMember& PartyMember);
-	UE_DEPRECATED(5.5, "This has been deprecated to support multiple local players and now requires a performing player.")
-	virtual bool TryKickMember(const UPartyMember& PartyMember) { return false; }
 	UE_API virtual bool TryKickMember(const ULocalPlayer& PerformingPlayer, const UPartyMember& PartyMember);
 	
 	UE_API bool TryInviteUser(const USocialUser& UserToInvite, const ESocialPartyInviteMethod InviteMethod = ESocialPartyInviteMethod::Other, const FString& MetaData = FString());
@@ -279,11 +265,7 @@ protected:
 	/** Virtual versions of the package-scoped "CanX" methods above, as a virtual declared within package scoping cannot link (exported public, imported protected) */
 	UE_API virtual ESocialPartyInviteFailureReason CanInviteUserInternal(const USocialUser& User, const ESocialPartyInviteMethod InviteMethod) const;
 
-	UE_DEPRECATED(5.5, "This has been deprecated to support multiple local players and now requires a performing player.")
-	virtual bool CanPromoteMemberInternal(const UPartyMember& PartyMember) const { return false; }
 	UE_API virtual bool CanPromoteMemberInternal(const ULocalPlayer& PerformingPlayer, const UPartyMember& PartyMember) const;
-	UE_DEPRECATED(5.5, "This has been deprecated to support multiple local players and now requires a performing player.")
-	virtual bool CanKickMemberInternal(const UPartyMember& PartyMember) const { return false; }
 	UE_API virtual bool CanKickMemberInternal(const ULocalPlayer& PerformingPlayer, const UPartyMember& PartyMember) const;
 
 	UE_API virtual void OnInviteSentInternal(ESocialSubsystem SubsystemType, const USocialUser& InvitedUser, bool bWasSuccessful, const ESocialPartyInviteFailureReason FailureReason, const ESocialPartyInviteMethod InviteMethod, const FString& MetaData);
@@ -411,10 +393,6 @@ private:
 	UPROPERTY()
 	TMap<FUniqueNetIdRepl, TObjectPtr<UPartyMember>> PartyMembersById;
 
-	UE_DEPRECATED(5.5, "Rejoinable Parties is deprecated and will be removed")
-	UPROPERTY(config)
-	bool bEnableAutomaticPartyRejoin = false;
-
 	TMap<FUniqueNetIdRepl, double> LastInviteSentById;
 
 	UPROPERTY(config)
@@ -526,17 +504,4 @@ PARTY_API TArray<FUniqueNetIdRepl> GetPartyMemberIds(const USocialParty& SocialP
 PARTY_API TArray<USocialToolkit*> GetLocalPartyMemberToolkits(const USocialParty& SocialParty);
 } // UE::OnlineFramework::Party
 
-namespace UE::OnlineFramework::Party
-{
-UE_DEPRECATED(5.5, "Use UE::OnlineFramework::GetPartyMemberIds")
-inline TArray<FUniqueNetIdRepl> GetPartyMemberIds(const USocialParty* SocialParty)
-{
-	if (SocialParty)
-	{
-		return UE::OnlineFramework::GetPartyMemberIds(*SocialParty);
-	}
-	return {};
-}
-} // UE::OnlineFramework::Party
-
 #undef UE_API
diff --git a/Engine/Plugins/Online/OnlineFramework/Source/Party/Public/SocialManager.h b/Engine/Plugins/Online/OnlineFramework/Source/Party/Public/SocialManager.h
--- a/Engine/Plugins/Online/OnlineFramework/Source/Party/Public/SocialManager.h
+++ b/Engine/Plugins/Online/OnlineFramework/Source/Party/Public/SocialManager.h
@@ -129,26 +129,10 @@ protected:
 	DECLARE_DELEGATE_OneParam(FOnJoinPartyAttemptComplete, const FJoinPartyResult&);
 	UE_API void JoinParty(const USocialUser& UserToJoin, const FOnlinePartyTypeId& PartyTypeId, const FOnJoinPartyAttemptComplete& OnJoinPartyComplete, const FName& JoinMethod);
 
-protected:
-	struct PARTY_API UE_DEPRECATED(5.5, "Rejoinable Parties is deprecated and will be removed") FRejoinableParty : public TSharedFromThis<FRejoinableParty>
-	{
-		FRejoinableParty(const USocialParty& SourceParty);
-
-		TSharedRef<const FOnlinePartyId> PartyId;
-		TArray<FUniqueNetIdRef> MemberIds;
-		FName OriginalJoinMethod;
-	};
-
 public:
 	struct FJoinPartyAttempt
 	{
-		UE_DEPRECATED(5.5, "Rejoinable Parties is deprecated and will be removed")
-UE_API PRAGMA_DISABLE_DEPRECATION_WARNINGS
-		FJoinPartyAttempt(TSharedRef<const FRejoinableParty> InRejoinInfo);
-		FJoinPartyAttempt(const FJoinPartyAttempt&) = default;
-		FJoinPartyAttempt(FJoinPartyAttempt&&) = default;
-UE_API PRAGMA_ENABLE_DEPRECATION_WARNINGS
-		FJoinPartyAttempt(const USocialUser* InTargetUser, const FOnlinePartyTypeId& InPartyTypeId, const FName& InJoinMethod, const FOnJoinPartyAttemptComplete& InOnJoinComplete);
+		UE_API FJoinPartyAttempt(const USocialUser* InTargetUser, const FOnlinePartyTypeId& InPartyTypeId, const FName& InJoinMethod, const FOnJoinPartyAttemptComplete& InOnJoinComplete);
 
 		UE_API FString ToDebugString() const;
 
@@ -157,10 +141,6 @@ UE_API PRAGMA_ENABLE_DEPRECATION_WARNINGS
 		FName JoinMethod = PartyJoinMethod::Unspecified;
 		FUniqueNetIdRepl TargetUserPlatformId;
 
-		UE_DEPRECATED(5.5, "Rejoinable Parties is deprecated and will be removed")
-PRAGMA_DISABLE_DEPRECATION_WARNINGS
-		TSharedPtr<const FRejoinableParty> RejoinInfo;
-PRAGMA_ENABLE_DEPRECATION_WARNINGS
 		IOnlinePartyJoinInfoConstPtr JoinInfo;
 
 		FOnJoinPartyAttemptComplete OnJoinComplete;
@@ -203,11 +183,6 @@ protected:
 	/** Up to the game to decide whether it wants to allow crossplay (generally based on a user setting of some kind) */
 	UE_API virtual ECrossplayPreference GetCrossplayPreference() const;
 
-	UE_DEPRECATED(5.5, "Rejoinable Parties is deprecated and will be removed")
-UE_API PRAGMA_DISABLE_DEPRECATION_WARNINGS
-	virtual bool ShouldTryRejoiningPersistentParty(const FRejoinableParty& InRejoinableParty) const;
-PRAGMA_ENABLE_DEPRECATION_WARNINGS
-
 	template <typename InteractionT>
 	void RegisterInteraction()
 	{
@@ -222,20 +197,11 @@ public:
 	UE_API const FJoinPartyAttempt* GetJoinAttemptInProgress(const FOnlinePartyTypeId& PartyTypeId) const;
 
 protected:
-	UE_DEPRECATED(5.5, "HandlePlatformSessionInviteAccepted is deprecated and will be removed")
-	void HandlePlatformSessionInviteAccepted(const FUniqueNetIdRef& LocalUserId, const FOnlineSessionSearchResult& InviteResult) {}
-
 	UE_API TSharedPtr<const IOnlinePartyJoinInfo> GetJoinInfoFromSession(const FOnlineSessionSearchResult& PlatformSession);
 	UE_API void FinishJoinPartyAttempt(FJoinPartyAttempt& JoinAttemptToDestroy, const FJoinPartyResult& JoinResult);
 
 	UE_API virtual TSubclassOf<USocialDebugTools> GetSocialDebugToolsClass() const;
 
-	/** Info on the persistent party we were in when losing connection to the party service and want to rejoin when it returns */
-	UE_DEPRECATED(5.5, "Rejoinable Parties is deprecated and will be removed")
-PRAGMA_DISABLE_DEPRECATION_WARNINGS
-	TSharedPtr<FRejoinableParty> RejoinableParty;
-PRAGMA_ENABLE_DEPRECATION_WARNINGS
-
 	/** The desired type of SocialToolkit to create for each local player */
 	TSubclassOf<USocialToolkit> ToolkitClass;
 
diff --git a/Engine/Plugins/Online/OnlineFramework/Source/Party/Public/SocialToolkit.h b/Engine/Plugins/Online/OnlineFramework/Source/Party/Public/SocialToolkit.h
--- a/Engine/Plugins/Online/OnlineFramework/Source/Party/Public/SocialToolkit.h
+++ b/Engine/Plugins/Online/OnlineFramework/Source/Party/Public/SocialToolkit.h
@@ -82,9 +82,6 @@ public:
 	UE_DEPRECATED(5.6, "AddLocalUserOnlineProperties is deprecated, please use IOnlinePresence::SetPresence instead.")
 	UE_API void AddLocalUserOnlineProperties(FPresenceProperties OnlineProperties);
 
-	UE_DEPRECATED(5.5, "GetOwningLocalPlayer returns LocalPlayerOwner which is a TWeakObjectPtr and may return nullptr when the local player logs out. Please use the pointer version.")
-	UE_API ULocalPlayer& GetOwningLocalPlayer() const;
-
 	UE_API ULocalPlayer* GetOwningLocalPlayerPtr() const;
 	UE_API USocialManager& GetSocialManager() const;
 
diff --git a/Engine/Plugins/Online/OnlineServices/Source/OnlineServicesCommon/Public/Online/OnlineServicesCommon.h b/Engine/Plugins/Online/OnlineServices/Source/OnlineServicesCommon/Public/Online/OnlineServicesCommon.h
--- a/Engine/Plugins/Online/OnlineServices/Source/OnlineServicesCommon/Public/Online/OnlineServicesCommon.h
+++ b/Engine/Plugins/Online/OnlineServices/Source/OnlineServicesCommon/Public/Online/OnlineServicesCommon.h
@@ -26,12 +26,6 @@ class FOnlineServicesCommon
 public:
 	using Super = IOnlineServices;
 	
-	UE_DEPRECATED(5.5, "Please call the new method taking an FName InstanceConfigName")
-	FOnlineServicesCommon(const FString& InServiceConfigName, FName InInstanceName)
-		: FOnlineServicesCommon(InServiceConfigName, InInstanceName, NAME_None)
-	{
-	}
-
 	UE_API FOnlineServicesCommon(const FString& InServiceConfigName, FName InInstanceName, FName InInstanceConfigName);
 	FOnlineServicesCommon(const FOnlineServicesCommon&) = delete;
 	FOnlineServicesCommon(FOnlineServicesCommon&&) = delete;
@@ -179,8 +173,6 @@ public:
 	/**
 	 * Get the ini config name for the Subsystem
 	 */
-	UE_DEPRECATED(5.5, "GetConfigName has been renamed GetServiceConfigName")
-	const FString& GetConfigName() const { return ServiceConfigName; }
 	const FString& GetServiceConfigName() const { return ServiceConfigName; }
 
 	UE_DEPRECATED(5.6, "GetConfigSectionHeiarchy has been renamed GetConfigSectionHeirarchy")
diff --git a/Engine/Plugins/Online/OnlineServices/Source/OnlineServicesInterface/Public/Online/OnlineServicesRegistry.h b/Engine/Plugins/Online/OnlineServices/Source/OnlineServicesInterface/Public/Online/OnlineServicesRegistry.h
--- a/Engine/Plugins/Online/OnlineServices/Source/OnlineServicesInterface/Public/Online/OnlineServicesRegistry.h
+++ b/Engine/Plugins/Online/OnlineServices/Source/OnlineServicesInterface/Public/Online/OnlineServicesRegistry.h
@@ -26,8 +26,6 @@ public:
 	 *
 	 * @return Initialized IOnlineServices instance
 	 */
-	UE_DEPRECATED(5.5, "Please use the new Create taking an additional InstanceConfigName param")
-	virtual TSharedPtr<IOnlineServices> Create(FName InstanceName) final { return Create(InstanceName, NAME_None); }
 	virtual TSharedPtr<IOnlineServices> Create(FName InstanceName, FName InstanceConfigName) = 0;
 };
 
@@ -83,8 +81,6 @@ public:
 	 * 
 	 * @return The services instance, or an invalid pointer if the OnlineServices is unavailable
 	 */
-	UE_DEPRECATED(5.5, "Please call the new GetNamedServicesInstance which takes an additional InstanceConfigName param")
-	TSharedPtr<IOnlineServices> GetNamedServicesInstance(EOnlineServices OnlineServices, FName InstanceName) { return GetNamedServicesInstance(OnlineServices, InstanceName, NAME_None); }
 	ONLINESERVICESINTERFACE_API TSharedPtr<IOnlineServices> GetNamedServicesInstance(EOnlineServices OnlineServices, FName InstanceName, FName InstanceConfigName);
 
 	/**
@@ -94,8 +90,6 @@ public:
 	 * @param InstanceName Name of the instance
 	 * @param InstanceConfigName Name of the config
 	 */
-	UE_DEPRECATED(5.5, "Please call the new DestroyNamedServicesInstance which takes an additional InstanceConfigName param")
-	void DestroyNamedServicesInstance(EOnlineServices OnlineServices, FName InstanceName) { return DestroyNamedServicesInstance(OnlineServices, InstanceName, NAME_None); }
 	ONLINESERVICESINTERFACE_API void DestroyNamedServicesInstance(EOnlineServices OnlineServices, FName InstanceName, FName InstanceConfigName);
 
 	/**
@@ -121,8 +115,6 @@ public:
 	 * 
 	 * @return The initialized IOnlineServices instance, or an invalid pointer if the OnlineServices is unavailable
 	 */
-	UE_DEPRECATED(5.5, "Please call the new CreateServices which takes an additional InstanceConfigName param")
-	TSharedPtr<IOnlineServices> CreateServices(EOnlineServices OnlineServices, FName InstanceName) { return CreateServices(OnlineServices, InstanceName, NAME_None); }
 	ONLINESERVICESINTERFACE_API TSharedPtr<IOnlineServices> CreateServices(EOnlineServices OnlineServices, FName InstanceName, FName InstanceConfigName);
 
 	/**
diff --git a/Engine/Plugins/Online/OnlineServicesEOSGS/Source/Public/Online/OnlineServicesEOSGS.h b/Engine/Plugins/Online/OnlineServicesEOSGS/Source/Public/Online/OnlineServicesEOSGS.h
--- a/Engine/Plugins/Online/OnlineServicesEOSGS/Source/Public/Online/OnlineServicesEOSGS.h
+++ b/Engine/Plugins/Online/OnlineServicesEOSGS/Source/Public/Online/OnlineServicesEOSGS.h
@@ -37,8 +37,6 @@ public:
 	ONLINESERVICESEOSGS_API IPlayerReportsPtr GetPlayerReportsInterface();
 	ONLINESERVICESEOSGS_API IPlayerSanctionsPtr GetPlayerSanctionsInterface(); 
 
-	UE_DEPRECATED(5.5, "Please call GetServiceConfigNameStatic instead")
-	static const TCHAR* GetConfigNameStatic() { return GetServiceConfigNameStatic(); }
 	static const TCHAR* GetServiceConfigNameStatic() { return TEXT("EOS"); }
 
 protected:
diff --git a/Engine/Plugins/Online/OnlineServicesNull/Source/Public/Online/AuthNull.h b/Engine/Plugins/Online/OnlineServicesNull/Source/Public/Online/AuthNull.h
--- a/Engine/Plugins/Online/OnlineServicesNull/Source/Public/Online/AuthNull.h
+++ b/Engine/Plugins/Online/OnlineServicesNull/Source/Public/Online/AuthNull.h
@@ -13,14 +13,6 @@ namespace UE::Online {
 
 class FOnlineServicesNull;
 
-class UE_DEPRECATED(5.5, "FOnlineAccountIdRegistryNull now uses TOnlineBasicAccountIdRegistry, and FOnlineAccountIdString was never exposed external to FOnlineAccountIdRegistryNull") FOnlineAccountIdString
-{
-public:
-	FString Data;
-	int32 AccountIndex;
-	FAccountId AccountId;
-};
-
 class FOnlineAccountIdRegistryNull
 	: public IOnlineAccountIdRegistry
 {
diff --git a/Engine/Plugins/Online/OnlineSubsystem/Source/Private/OnlinePartyInterface.cpp b/Engine/Plugins/Online/OnlineSubsystem/Source/Private/OnlinePartyInterface.cpp
--- a/Engine/Plugins/Online/OnlineSubsystem/Source/Private/OnlinePartyInterface.cpp
+++ b/Engine/Plugins/Online/OnlineSubsystem/Source/Private/OnlinePartyInterface.cpp
@@ -453,12 +453,6 @@ const TCHAR* ToString(const EJoinPartyCompletionResult Value)
 	{
 		return TEXT("LoggedOut");
 	}
-PRAGMA_DISABLE_DEPRECATION_WARNINGS
-	case EJoinPartyCompletionResult::UnableToRejoin:
-PRAGMA_ENABLE_DEPRECATION_WARNINGS
-	{
-		return TEXT("UnableToRejoin");
-	}
 	case EJoinPartyCompletionResult::IncompatiblePlatform:
 	{
 		return TEXT("IncompatiblePlatform");
@@ -907,13 +901,3 @@ FString IOnlinePartyJoinInfo::ToDebugString() const
 		IsAcceptingMembers() ? 1 : 0,
 		GetNotAcceptingReason());
 };
-
-bool IOnlinePartySystem::RejoinParty(const FUniqueNetId& LocalUserId, const FOnlinePartyId& PartyId, const FOnlinePartyTypeId& PartyTypeId, const TArray<FUniqueNetIdRef>& FormerMembers, const FOnJoinPartyComplete& Delegate)
-{
-	return false;
-}
-
-bool IOnlinePartySystem::LeaveParty(const FUniqueNetId& LocalUserId, const FOnlinePartyId& PartyId, const FOnLeavePartyComplete& Delegate)
-{
-	return LeaveParty(LocalUserId, PartyId, true, Delegate);
-}
diff --git a/Engine/Plugins/Online/OnlineSubsystem/Source/Public/Interfaces/OnlinePartyInterface.h b/Engine/Plugins/Online/OnlineSubsystem/Source/Public/Interfaces/OnlinePartyInterface.h
--- a/Engine/Plugins/Online/OnlineSubsystem/Source/Public/Interfaces/OnlinePartyInterface.h
+++ b/Engine/Plugins/Online/OnlineSubsystem/Source/Public/Interfaces/OnlinePartyInterface.h
@@ -623,7 +623,6 @@ enum class EPartyState : uint8
 	None,
 	CreatePending,
 	JoinPending,
-	RejoinPending UE_DEPRECATED(5.5, "RejoinParty is deprecated. Implement this feature via JoinParty."), 
 	LeavePending,
 	Active,
 	Disconnected,
@@ -1257,33 +1256,6 @@ public:
 	 */
 	virtual void QueryPartyJoinability(const FUniqueNetId& LocalUserId, const IOnlinePartyJoinInfo& OnlinePartyJoinInfo, const FOnQueryPartyJoinabilityCompleteEx& Delegate) = 0;
 
-	/**
-	 * DEPRECATED! Attempt to rejoin a former party. Deprecated: Implement this feature with JoinParty instead.
-	 *
-	 * @param LocalUserId - user making the request
-	 * @param PartyId - id of the party you want to rejoin
-	 * @param PartyTypeId - type id of the party you want to rejoin
-	 * @param FormerMembers - array of former member ids that we can contact to try to rejoin the party
-	 * @param Delegate - called on completion
-	 *
-	 * @return true if task was started
-	 */
-	UE_DEPRECATED(5.5, "RejoinParty is deprecated. Implement this feature via JoinParty.")
-	virtual bool ONLINESUBSYSTEM_API RejoinParty(const FUniqueNetId& LocalUserId, const FOnlinePartyId& PartyId, const FOnlinePartyTypeId& PartyTypeId, const TArray<FUniqueNetIdRef>& FormerMembers, const FOnJoinPartyComplete& Delegate = FOnJoinPartyComplete());
-
-	/**
-	 * Leave an existing party
-	 * All existing party members notified of member leaving (see FOnPartyMemberLeft)
-	 *
-	 * @param LocalUserId - user making the request
-	 * @param PartyId - id of an existing party
-	 * @param Delegate - called on completion
-	 *
-	 * @return true if task was started
-	 */
-	UE_DEPRECATED(5.5, "Use LeaveParty that takes bSynchronizeLeave")
-	virtual bool ONLINESUBSYSTEM_API LeaveParty(const FUniqueNetId& LocalUserId, const FOnlinePartyId& PartyId, const FOnLeavePartyComplete& Delegate = FOnLeavePartyComplete());
-
 	/**
 	 * Leave an existing party
 	 * All existing party members notified of member leaving (see FOnPartyMemberLeft)
@@ -1888,8 +1860,6 @@ enum class EJoinPartyCompletionResult : int8
 	NoResponse,
 	/** You were logged out while attempting to join the party */
 	LoggedOut,
-	/** You were unable to rejoin the party */
-	UnableToRejoin UE_DEPRECATED(5.5, "RejoinParty is deprecated. Implement this feature via JoinParty."),
 	/** Your platform is not compatible with the party */
 	IncompatiblePlatform,
 
diff --git a/Engine/Plugins/Online/OnlineSubsystem/Source/Public/OnlineStats.h b/Engine/Plugins/Online/OnlineSubsystem/Source/Public/OnlineStats.h
--- a/Engine/Plugins/Online/OnlineSubsystem/Source/Public/OnlineStats.h
+++ b/Engine/Plugins/Online/OnlineSubsystem/Source/Public/OnlineStats.h
@@ -20,47 +20,6 @@ public:
 	FORCEINLINE FStatsColumnArray(const FStatsColumnArray& Other) : Super(Other) {}
 	FORCEINLINE FStatsColumnArray& operator=(FStatsColumnArray&& Other) { Super::operator=(MoveTemp(Other)); return *this; }
 	FORCEINLINE FStatsColumnArray& operator=(const FStatsColumnArray& Other) { Super::operator=(Other); return *this; }
-
-	UE_DEPRECATED(5.5, "FStatsColumnArray now uses FString for the key type instead of FName")
-		FStatsColumnArray(const FOnlineKeyValuePairs< FName, FVariantData >& DeprecatedValues)
-	{
-		Reserve(DeprecatedValues.Num());
-		for (const TPair< FName, FVariantData >& DeprecatedValue : DeprecatedValues)
-		{
-			Super::Emplace(DeprecatedValue.Key.ToString(), DeprecatedValue.Value);
-		}
-	}
-
-	UE_DEPRECATED(5.5, "FStatsColumnArray now uses FString for the key type instead of FName")
-		FStatsColumnArray(FOnlineKeyValuePairs< FName, FVariantData >&& DeprecatedValues)
-	{
-		Reserve(DeprecatedValues.Num());
-		for (TPair< FName, FVariantData >& DeprecatedValue : DeprecatedValues)
-		{
-			Super::Emplace(DeprecatedValue.Key.ToString(), MoveTemp(DeprecatedValue.Value));
-		}
-	}
-
-	using Super::Add;
-	UE_DEPRECATED(5.5, "FStatsColumnArray now uses FString for the key type instead of FName")
-		FVariantData& Add(const FName& InKey, const FVariantData& InValue)
-	{
-		return Super::Add(InKey.ToString(), InValue);
-	}
-
-	using Super::Emplace;
-	UE_DEPRECATED(5.5, "FStatsColumnArray now uses FString for the key type instead of FName")
-		FVariantData& Emplace(const FName& InKey, const FVariantData& InValue)
-	{
-		return Super::Emplace(InKey.ToString(), InValue);
-	}
-
-	using Super::Find;
-	UE_DEPRECATED(5.5, "FStatsColumnArray now uses FString for the key type instead of FName")
-		FVariantData* Find(const FName& InKey)
-	{
-		return Super::Find(InKey.ToString());
-	}
 };
 
 
@@ -77,47 +36,6 @@ public:
 	FORCEINLINE FStatPropertyArray(const FStatPropertyArray& Other) : Super(Other) {}
 	FORCEINLINE FStatPropertyArray& operator=(FStatPropertyArray&& Other) { Super::operator=(MoveTemp(Other)); return *this; }
 	FORCEINLINE FStatPropertyArray& operator=(const FStatPropertyArray& Other) { Super::operator=(Other); return *this; }
-
-	UE_DEPRECATED(5.5, "FStatPropertyArray now uses FString for the key type instead of FName")
-		FStatPropertyArray(const FOnlineKeyValuePairs< FName, FVariantData >& DeprecatedValues)
-	{
-		Reserve(DeprecatedValues.Num());
-		for (const TPair< FName, FVariantData >& DeprecatedValue : DeprecatedValues)
-		{
-			Super::Emplace(DeprecatedValue.Key.ToString(), DeprecatedValue.Value);
-		}
-	}
-
-	UE_DEPRECATED(5.5, "FStatPropertyArray now uses FString for the key type instead of FName")
-		FStatPropertyArray(FOnlineKeyValuePairs< FName, FVariantData >&& DeprecatedValues)
-	{
-		Reserve(DeprecatedValues.Num());
-		for (TPair< FName, FVariantData >& DeprecatedValue : DeprecatedValues)
-		{
-			Super::Emplace(DeprecatedValue.Key.ToString(), MoveTemp(DeprecatedValue.Value));
-		}
-	}
-
-	using Super::Add;
-	UE_DEPRECATED(5.5, "FStatPropertyArray now uses FString for the key type instead of FName")
-		FVariantData& Add(const FName& InKey, const FVariantData& InValue)
-	{
-		return Super::Add(InKey.ToString(), InValue);
-	}
-
-	using Super::Emplace;
-	UE_DEPRECATED(5.5, "FStatPropertyArray now uses FString for the key type instead of FName")
-		FVariantData& Emplace(const FName& InKey, const FVariantData& InValue)
-	{
-		return Super::Emplace(InKey.ToString(), InValue);
-	}
-
-	using Super::Find;
-	UE_DEPRECATED(5.5, "FStatPropertyArray now uses FString for the key type instead of FName")
-		FVariantData* Find(const FName& InKey)
-	{
-		return Super::Find(InKey.ToString());
-	}
 };
 
 class FNameArrayDeprecationWrapper : public TArray<FString>
@@ -135,20 +53,6 @@ public:
 	FORCEINLINE FNameArrayDeprecationWrapper(const TArray<FString>& Other) : Super(Other) {}
 	FORCEINLINE FNameArrayDeprecationWrapper& operator=(TArray<FString>&& Other) { Super::operator=(MoveTemp(Other)); return *this; }
 	FORCEINLINE FNameArrayDeprecationWrapper& operator=(const TArray<FString>& Other) { Super::operator=(Other); return *this; }
-
-	using Super::Add;
-	UE_DEPRECATED(5.5, "This variable is now a TArray<FString> instead of a TArray<FName>.")
-		int32 Add(const FName& InElement)
-	{
-		return Super::Add(InElement.ToString());
-	}
-
-	using Super::Emplace;
-	UE_DEPRECATED(5.5, "This variable is now a TArray<FString> instead of a TArray<FName>.")
-		int32 Emplace(const FName& InElement)
-	{
-		return Super::Emplace(InElement.ToString());
-	}
 };
 
 class FNameDeprecationWrapper : public FString
@@ -166,24 +70,6 @@ public:
 	FORCEINLINE FNameDeprecationWrapper(const FString& Other) : Super(Other) {}
 	FORCEINLINE FNameDeprecationWrapper& operator=(FString&& Other) { Super::operator=(MoveTemp(Other)); return *this; }
 	FORCEINLINE FNameDeprecationWrapper& operator=(const FString& Other) { Super::operator=(Other); return *this; }
-
-	UE_DEPRECATED(5.5, "This variable is now an FString instead of an FName.")
-	FORCEINLINE FNameDeprecationWrapper(FName&& DeprecatedValue) : Super(DeprecatedValue.ToString()) {}
-
-	UE_DEPRECATED(5.5, "This variable is now an FString instead of an FName.")
-		FORCEINLINE FNameDeprecationWrapper(const FName& DeprecatedValue) : Super(DeprecatedValue.ToString()) {}
-
-	UE_DEPRECATED(5.5, "This variable is now an FString instead of an FName.")
-		FORCEINLINE FNameDeprecationWrapper& operator=(FName&& DeprecatedValue) { Super::operator=(DeprecatedValue.ToString()); return *this; }
-
-	UE_DEPRECATED(5.5, "This variable is now an FString instead of an FName.")
-		FORCEINLINE FNameDeprecationWrapper& operator=(const FName& DeprecatedValue) { Super::operator=(DeprecatedValue.ToString()); return *this; }
-
-	UE_DEPRECATED(5.5, "This variable is now an FString instead of an FName.")
-		FORCEINLINE FString ToString() const { return *this; }
-
-	UE_DEPRECATED(5.5, "This variable is now an FString instead of an FName.")
-		FORCEINLINE void ToString(FString& Out) const { Out = *this; }
 };
 
 /**
@@ -257,21 +143,6 @@ public:
 	 */
 	UE_API virtual void DecrementIntStat(const FString& StatName, int32 DecBy = 1);
 
-	UE_DEPRECATED(5.5, "Use the FString overload instead")
-	FVariantData* FindStatByName(const FName& StatName) { return FindStatByName(StatName.ToString()); }
-	UE_DEPRECATED(5.5, "Use the FString overload instead")
-	virtual void SetFloatStat(const FName& StatName, float Value) { SetFloatStat(StatName.ToString(), Value); }
-	UE_DEPRECATED(5.5, "Use the FString overload instead")
-	virtual void SetIntStat(const FName& StatName, int32 Value) { SetIntStat(StatName.ToString(), Value); }
-	UE_DEPRECATED(5.5, "Use the FString overload instead")
-	virtual void IncrementFloatStat(const FName& StatName, float IncBy = 1.0f) { IncrementFloatStat(StatName.ToString(), IncBy); }
-	UE_DEPRECATED(5.5, "Use the FString overload instead")
-	virtual void IncrementIntStat(const FName& StatName, int32 IncBy = 1) { IncrementIntStat(StatName.ToString(), IncBy); }
-	UE_DEPRECATED(5.5, "Use the FString overload instead")
-	virtual void DecrementFloatStat(const FName& StatName, float DecBy = 1.0f) { DecrementFloatStat(StatName.ToString(), DecBy); }
-	UE_DEPRECATED(5.5, "Use the FString overload instead")
-	virtual void DecrementIntStat(const FName& StatName, int32 DecBy = 1) { DecrementIntStat(StatName.ToString(), DecBy); }
-
 	/**
 	 * Destructor
 	 */
@@ -354,13 +225,6 @@ public:
 	/** Type of data this column represents */
 	const EOnlineKeyValuePairDataType::Type DataType;
 
-	UE_DEPRECATED(5.5, "Use the FString overload instead")
-	FColumnMetaData(const FName InColumnName, EOnlineKeyValuePairDataType::Type InDataType) :
-		ColumnName(InColumnName.ToString()),
-		DataType(InDataType)
-	{
-	}
-
 	FColumnMetaData(const FString& InColumnName, EOnlineKeyValuePairDataType::Type InDataType) :
 		ColumnName(InColumnName),
 		DataType(InDataType)
diff --git a/Engine/Plugins/Online/OnlineSubsystemEOS/Source/OnlineSubsystemEOS/Private/EOSSettings.cpp b/Engine/Plugins/Online/OnlineSubsystemEOS/Source/OnlineSubsystemEOS/Private/EOSSettings.cpp
--- a/Engine/Plugins/Online/OnlineSubsystemEOS/Source/OnlineSubsystemEOS/Private/EOSSettings.cpp
+++ b/Engine/Plugins/Online/OnlineSubsystemEOS/Source/OnlineSubsystemEOS/Private/EOSSettings.cpp
@@ -148,12 +148,6 @@ FEOSSettings::FEOSSettings()
 	, bUseEAS(false)
 	, bUseEOSConnect(false)
 	, bUseEOSRTC(true)
-PRAGMA_DISABLE_DEPRECATION_WARNINGS
-	, bUseEOSSessions(false)
-	, bMirrorStatsToEOS(false)
-	, bMirrorAchievementsToEOS(false)
-	, bMirrorPresenceToEAS(false)
-PRAGMA_ENABLE_DEPRECATION_WARNINGS
 	, bUseNewLoginFlow(false)
 {
 
@@ -199,12 +193,6 @@ const FEOSSettings& UEOSSettings::ManualGetSettings()
 		GConfig->GetBool(INI_SECTION, TEXT("bUseEAS"), GCachedSettings->bUseEAS, GEngineIni);
 		GConfig->GetBool(INI_SECTION, TEXT("bUseEOSConnect"), GCachedSettings->bUseEOSConnect, GEngineIni);
 		GConfig->GetBool(INI_SECTION, TEXT("bUseEOSRTC"), GCachedSettings->bUseEOSRTC, GEngineIni);
-PRAGMA_DISABLE_DEPRECATION_WARNINGS
-		GConfig->GetBool(INI_SECTION, TEXT("bUseEOSSessions"), GCachedSettings->bUseEOSSessions, GEngineIni);
-		GConfig->GetBool(INI_SECTION, TEXT("bMirrorStatsToEOS"), GCachedSettings->bMirrorStatsToEOS, GEngineIni);
-		GConfig->GetBool(INI_SECTION, TEXT("bMirrorAchievementsToEOS"), GCachedSettings->bMirrorAchievementsToEOS, GEngineIni);
-		GConfig->GetBool(INI_SECTION, TEXT("bMirrorPresenceToEAS"), GCachedSettings->bMirrorPresenceToEAS, GEngineIni);
-PRAGMA_ENABLE_DEPRECATION_WARNINGS
 		GConfig->GetBool(INI_SECTION, TEXT("bUseNewLoginFlow"), GCachedSettings->bUseNewLoginFlow, GEngineIni);
 		// Artifacts explicitly skipped
 		GConfig->GetArray(INI_SECTION, TEXT("TitleStorageTags"), GCachedSettings->TitleStorageTags, GEngineIni);
@@ -235,12 +223,6 @@ FEOSSettings UEOSSettings::ToNative() const
 	Native.bUseEAS = bUseEAS;
 	Native.bUseEOSConnect = bUseEOSConnect;
 	Native.bUseEOSRTC = bUseEOSRTC;
-PRAGMA_DISABLE_DEPRECATION_WARNINGS
-	Native.bUseEOSSessions = bUseEOSSessions;
-	Native.bMirrorStatsToEOS = bMirrorStatsToEOS;
-	Native.bMirrorAchievementsToEOS = bMirrorAchievementsToEOS;
-	Native.bMirrorPresenceToEAS = bMirrorPresenceToEAS;
-PRAGMA_ENABLE_DEPRECATION_WARNINGS
 	Native.bUseNewLoginFlow = bUseNewLoginFlow;
 	Algo::Transform(Artifacts, Native.Artifacts, &FArtifactSettings::ToNative);
 	Native.TitleStorageTags = TitleStorageTags;
@@ -486,54 +468,6 @@ void UEOSSettings::PostEditChangeProperty(FPropertyChangedEvent& PropertyChanged
 		}
 	}
 
-	// Turning off EAS disables presence mirroring too
-	if (PropertyChangedEvent.Property->GetFName() == FName(TEXT("bUseEAS")))
-	{
-		if (!bUseEAS)
-		{
-PRAGMA_DISABLE_DEPRECATION_WARNINGS
-			bMirrorPresenceToEAS = false;
-PRAGMA_ENABLE_DEPRECATION_WARNINGS
-		}
-	}
-
-	// Turning on presence requires EAS
-	if (PropertyChangedEvent.Property->GetFName() == FName(TEXT("bMirrorPresenceToEAS")))
-	{
-PRAGMA_DISABLE_DEPRECATION_WARNINGS
-		if (bMirrorPresenceToEAS)
-PRAGMA_ENABLE_DEPRECATION_WARNINGS
-		{
-			bUseEAS = true;
-		}
-	}
-
-	// Turning off EAS disables presence mirroring too
-	if (PropertyChangedEvent.Property->GetFName() == FName(TEXT("bUseEOSConnect")))
-	{
-		if (!bUseEOSConnect)
-		{
-PRAGMA_DISABLE_DEPRECATION_WARNINGS
-			bMirrorAchievementsToEOS = false;
-			bMirrorStatsToEOS = false;
-			bUseEOSSessions = false;
-PRAGMA_ENABLE_DEPRECATION_WARNINGS
-		}
-	}
-
-	// These all require EOS turned on if they are on
-	if (PropertyChangedEvent.Property->GetFName() == FName(TEXT("bMirrorAchievementsToEOS")) ||
-		PropertyChangedEvent.Property->GetFName() == FName(TEXT("bMirrorStatsToEOS")) ||
-		PropertyChangedEvent.Property->GetFName() == FName(TEXT("bUseEOSSessions")))
-	{
-PRAGMA_DISABLE_DEPRECATION_WARNINGS
-		if (bMirrorAchievementsToEOS || bMirrorStatsToEOS || bUseEOSSessions)
-PRAGMA_ENABLE_DEPRECATION_WARNINGS
-		{
-			bUseEOSConnect = true;
-		}
-	}
-
 	Super::PostEditChangeProperty(PropertyChangedEvent);
 }
 
diff --git a/Engine/Plugins/Online/OnlineSubsystemEOS/Source/OnlineSubsystemEOS/Private/OnlineSessionEOS.cpp b/Engine/Plugins/Online/OnlineSubsystemEOS/Source/OnlineSubsystemEOS/Private/OnlineSessionEOS.cpp
--- a/Engine/Plugins/Online/OnlineSubsystemEOS/Source/OnlineSubsystemEOS/Private/OnlineSessionEOS.cpp
+++ b/Engine/Plugins/Online/OnlineSubsystemEOS/Source/OnlineSubsystemEOS/Private/OnlineSessionEOS.cpp
@@ -2586,9 +2586,7 @@ uint32 FOnlineSessionEOS::FindEOSSession(int32 SearchingPlayerNum, const TShared
 		const FOnlineSessionSearchParam& SearchParam = It.Value();
 
 		// Game server keys are skipped
-PRAGMA_DISABLE_DEPRECATION_WARNINGS
-		if (Key == SEARCH_DEDICATED_ONLY || Key == SETTING_MAPNAME || Key == SEARCH_EMPTY_SERVERS_ONLY || Key == SEARCH_SECURE_SERVERS_ONLY || Key == SEARCH_PRESENCE || Key == SEARCH_LOBBIES)
-PRAGMA_ENABLE_DEPRECATION_WARNINGS
+		if (Key == SEARCH_DEDICATED_ONLY || Key == SETTING_MAPNAME || Key == SEARCH_EMPTY_SERVERS_ONLY || Key == SEARCH_SECURE_SERVERS_ONLY || Key == SEARCH_LOBBIES)
 		{
 			continue;
 		}
@@ -4699,9 +4697,7 @@ uint32 FOnlineSessionEOS::FindLobbySession(int32 SearchingPlayerNum, const TShar
 			const FOnlineSessionSearchParam& SearchParam = It.Value();
 
 			// Game server keys are skipped
-PRAGMA_DISABLE_DEPRECATION_WARNINGS
-			if (Key == SEARCH_DEDICATED_ONLY || Key == SETTING_MAPNAME || Key == SEARCH_EMPTY_SERVERS_ONLY || Key == SEARCH_SECURE_SERVERS_ONLY || Key == SEARCH_PRESENCE || Key == SEARCH_LOBBIES)
-PRAGMA_ENABLE_DEPRECATION_WARNINGS
+			if (Key == SEARCH_DEDICATED_ONLY || Key == SETTING_MAPNAME || Key == SEARCH_EMPTY_SERVERS_ONLY || Key == SEARCH_SECURE_SERVERS_ONLY || Key == SEARCH_LOBBIES)
 			{
 				continue;
 			}
diff --git a/Engine/Plugins/Online/OnlineSubsystemEOS/Source/OnlineSubsystemEOS/Public/EOSSettings.h b/Engine/Plugins/Online/OnlineSubsystemEOS/Source/OnlineSubsystemEOS/Public/EOSSettings.h
--- a/Engine/Plugins/Online/OnlineSubsystemEOS/Source/OnlineSubsystemEOS/Public/EOSSettings.h
+++ b/Engine/Plugins/Online/OnlineSubsystemEOS/Source/OnlineSubsystemEOS/Public/EOSSettings.h
@@ -90,14 +90,6 @@ PRAGMA_ENABLE_DEPRECATION_WARNINGS
 	bool bUseEAS;
 	bool bUseEOSConnect;
 	bool bUseEOSRTC;
-	UE_DEPRECATED(5.6, "EOSPlus is deprecated, and related settings will be removed soon.")
-	bool bUseEOSSessions;
-	UE_DEPRECATED(5.5, "EOSPlus is deprecated, and related settings will be removed soon.")
-	bool bMirrorStatsToEOS;
-	UE_DEPRECATED(5.5, "EOSPlus is deprecated, and related settings will be removed soon.")
-	bool bMirrorAchievementsToEOS;
-	UE_DEPRECATED(5.5, "EOSPlus is deprecated, and related settings will be removed soon. Equivalent functionality can be accessed via the IntegratedPlatformManagementFlags config.")
-	bool bMirrorPresenceToEAS;
 	bool bUseNewLoginFlow;
 	TArray<FEOSArtifactSettings> Artifacts;
 	TArray<FString> TitleStorageTags;
@@ -175,26 +167,6 @@ public:
 	UPROPERTY(Config, EditAnywhere, BlueprintReadOnly, Category="EOS Settings", DisplayName="Enable Real-Time chat")
 	bool bUseEOSRTC = true;
 
-	/** Set to true to write stats to EOS as well as the default platform */
-	UE_DEPRECATED(5.5, "EOSPlus is deprecated, and related settings will be removed soon.")
-	UPROPERTY(Config, EditAnywhere, BlueprintReadOnly, Category="Crossplay Settings")
-	bool bMirrorStatsToEOS = false;
-
-	/** Set to true to write achievement data to EOS as well as the default platform */
-	UE_DEPRECATED(5.5, "EOSPlus is deprecated, and related settings will be removed soon.")
-	UPROPERTY(Config, EditAnywhere, BlueprintReadOnly, Category="Crossplay Settings")
-	bool bMirrorAchievementsToEOS = false;
-
-	/** Set to true to use EOS for session registration with data mirrored to the default platform */
-	UE_DEPRECATED(5.6, "EOSPlus is deprecated, and related settings will be removed soon.")
-	UPROPERTY(Config, EditAnywhere, BlueprintReadOnly, Category="Crossplay Settings", DisplayName="Use Crossplay Sessions")
-	bool bUseEOSSessions = false;
-
-	/** Set to true to have Epic Accounts presence information updated when the default platform is updated */
-	UE_DEPRECATED(5.5, "EOSPlus is deprecated, and related settings will be removed soon. Equivalent functionality can be accessed via the IntegratedPlatformManagementFlags config.")
-	UPROPERTY(Config, EditAnywhere, BlueprintReadOnly, Category="Crossplay Settings")
-	bool bMirrorPresenceToEAS = false;
-
 	/** Set to true to use new EOS login flow */
 	UPROPERTY(Config, EditAnywhere, BlueprintReadOnly, Category="EOS Login Settings", DisplayName="Use new Login flow, which doesn't rely on EOSPlus")
 	bool bUseNewLoginFlow = false;
diff --git a/Engine/Plugins/Online/OnlineSubsystemEOS/Source/OnlineSubsystemEOSPlus/Private/OnlineAchievementsEOSPlus.cpp b/Engine/Plugins/Online/OnlineSubsystemEOS/Source/OnlineSubsystemEOSPlus/Private/OnlineAchievementsEOSPlus.cpp
--- a/Engine/Plugins/Online/OnlineSubsystemEOS/Source/OnlineSubsystemEOSPlus/Private/OnlineAchievementsEOSPlus.cpp
+++ b/Engine/Plugins/Online/OnlineSubsystemEOS/Source/OnlineSubsystemEOSPlus/Private/OnlineAchievementsEOSPlus.cpp
@@ -51,16 +51,6 @@ void FOnlineAchievementsEOSPlus::WriteAchievements(const FUniqueNetId& PlayerId,
 	{
 		BaseAchievementsInterface->WriteAchievements(*NetIdPlus->GetBaseNetId(), WriteObject, Delegate);
 	}
-PRAGMA_DISABLE_DEPRECATION_WARNINGS
-	if (UEOSSettings::GetSettings().bMirrorAchievementsToEOS)
-PRAGMA_ENABLE_DEPRECATION_WARNINGS
-	{
-		if (NetIdPlus->GetEOSNetId().IsValid())
-		{
-			// Mirror the achievement data to EOS
-			EosAchievementsInterface->WriteAchievements(*NetIdPlus->GetEOSNetId(), WriteObject, Ignored);
-		}
-	}
 }
 
 void FOnlineAchievementsEOSPlus::QueryAchievements(const FUniqueNetId& PlayerId, const FOnQueryAchievementsCompleteDelegate& Delegate)
diff --git a/Engine/Plugins/Online/OnlineSubsystemEOS/Source/OnlineSubsystemEOSPlus/Private/OnlineSessionEOSPlus.cpp b/Engine/Plugins/Online/OnlineSubsystemEOS/Source/OnlineSubsystemEOSPlus/Private/OnlineSessionEOSPlus.cpp
--- a/Engine/Plugins/Online/OnlineSubsystemEOS/Source/OnlineSubsystemEOSPlus/Private/OnlineSessionEOSPlus.cpp
+++ b/Engine/Plugins/Online/OnlineSubsystemEOS/Source/OnlineSubsystemEOSPlus/Private/OnlineSessionEOSPlus.cpp
@@ -16,10 +16,6 @@ FOnlineSessionEOSPlus::FOnlineSessionEOSPlus(FOnlineSubsystemEOSPlus* InSubsyste
 	EOSSessionInterface = EOSPlus->EosOSS->GetSessionInterface();
 	check(BaseSessionInterface.IsValid() && EOSSessionInterface.IsValid());
 
-PRAGMA_DISABLE_DEPRECATION_WARNINGS
-	bUseEOSSessions = UEOSSettings::GetSettings().bUseEOSSessions;
-PRAGMA_ENABLE_DEPRECATION_WARNINGS
-
 	BaseSessionInterface->AddOnSessionUserInviteAcceptedDelegate_Handle(FOnSessionUserInviteAcceptedDelegate::CreateRaw(this, &FOnlineSessionEOSPlus::OnSessionUserInviteAcceptedBase));
 	BaseSessionInterface->AddOnSessionInviteReceivedDelegate_Handle(FOnSessionInviteReceivedDelegate::CreateRaw(this, &FOnlineSessionEOSPlus::OnSessionInviteReceivedBase));
 	BaseSessionInterface->AddOnCreateSessionCompleteDelegate_Handle(FOnCreateSessionCompleteDelegate::CreateRaw(this, &FOnlineSessionEOSPlus::OnCreateSessionComplete));
diff --git a/Engine/Plugins/Online/OnlineSubsystemEOS/Source/OnlineSubsystemEOSPlus/Private/OnlineStatsEOSPlus.cpp b/Engine/Plugins/Online/OnlineSubsystemEOS/Source/OnlineSubsystemEOSPlus/Private/OnlineStatsEOSPlus.cpp
--- a/Engine/Plugins/Online/OnlineSubsystemEOS/Source/OnlineSubsystemEOSPlus/Private/OnlineStatsEOSPlus.cpp
+++ b/Engine/Plugins/Online/OnlineSubsystemEOS/Source/OnlineSubsystemEOSPlus/Private/OnlineStatsEOSPlus.cpp
@@ -67,31 +67,6 @@ void FOnlineStatsEOSPlus::UpdateStats(const FUniqueNetIdRef LocalUserId, const T
 
 		const FEOSSettings& EOSSettings = UEOSSettings::GetSettings();
 
-PRAGMA_DISABLE_DEPRECATION_WARNINGS
-		if (EOSSettings.bMirrorStatsToEOS)
-PRAGMA_ENABLE_DEPRECATION_WARNINGS
-		{
-			// Also write the data to EOS
-			IOnlineStatsPtr EOSStats = EOSPlus->EosOSS->GetStatsInterface();
-			if (EOSStats.IsValid())
-			{
-				// We send a processed copy of the updated stats to the EOS interface, with the corresponding EOS user ids
-				TArray<FOnlineStatsUserUpdatedStats> EOSUpdatedUserStats;
-
-				for (const FOnlineStatsUserUpdatedStats& UserStat : UpdatedUserStats)
-				{
-					FUniqueNetIdEOSPlusPtr UserStatNetIdPlus = GetNetIdPlus(UserStat.Account->ToString());
-
-					if (UserStatNetIdPlus.IsValid())
-					{
-						FOnlineStatsUserUpdatedStats& EOSUpdatedStats = EOSUpdatedUserStats.Emplace_GetRef(UserStatNetIdPlus->GetEOSNetId().ToSharedRef());
-						EOSUpdatedStats.Stats = UserStat.Stats;
-					}
-				}
-
-				EOSStats->UpdateStats(NetIdPlus->GetEOSNetId().ToSharedRef(), EOSUpdatedUserStats, IgnoredStatsComplete);
-			}
-		}
 	}
 }
 
diff --git a/Engine/Plugins/Online/OnlineSubsystemEOS/Source/OnlineSubsystemEOSPlus/Private/OnlineUserEOSPlus.cpp b/Engine/Plugins/Online/OnlineSubsystemEOS/Source/OnlineSubsystemEOSPlus/Private/OnlineUserEOSPlus.cpp
--- a/Engine/Plugins/Online/OnlineSubsystemEOS/Source/OnlineSubsystemEOSPlus/Private/OnlineUserEOSPlus.cpp
+++ b/Engine/Plugins/Online/OnlineSubsystemEOS/Source/OnlineSubsystemEOSPlus/Private/OnlineUserEOSPlus.cpp
@@ -1560,10 +1560,7 @@ void FOnlineUserEOSPlus::SetPresence(const FUniqueNetId& User, const FOnlineUser
 	{
 		// Skip setting EAS presence if not mirrored or if we errored at the platform level or the EOS user isn't found
 				
-		if (!bWasSuccessful || !NetIdPlus->GetEOSNetId().IsValid()
-PRAGMA_DISABLE_DEPRECATION_WARNINGS
-			|| !UEOSSettings::GetSettings().bMirrorPresenceToEAS)
-PRAGMA_ENABLE_DEPRECATION_WARNINGS
+		if (!bWasSuccessful || !NetIdPlus->GetEOSNetId().IsValid())
 		{
 			IntermediateComplete.ExecuteIfBound(UserId, bWasSuccessful);
 			return;
diff --git a/Engine/Plugins/Online/OnlineSubsystemSteam/Source/Private/OnlineSessionAsyncLobbySteam.cpp b/Engine/Plugins/Online/OnlineSubsystemSteam/Source/Private/OnlineSessionAsyncLobbySteam.cpp
--- a/Engine/Plugins/Online/OnlineSubsystemSteam/Source/Private/OnlineSessionAsyncLobbySteam.cpp
+++ b/Engine/Plugins/Online/OnlineSubsystemSteam/Source/Private/OnlineSessionAsyncLobbySteam.cpp
@@ -850,9 +850,7 @@ void FOnlineAsyncTaskSteamFindLobbiesBase::CreateQuery()
 		const FOnlineSessionSearchParam& SearchParam = It.Value();
 
 		// Game server keys are skipped
-PRAGMA_DISABLE_DEPRECATION_WARNINGS
-		if (Key == SEARCH_DEDICATED_ONLY || Key == SETTING_MAPNAME || Key == SEARCH_EMPTY_SERVERS_ONLY || Key == SEARCH_SECURE_SERVERS_ONLY || Key == SEARCH_PRESENCE || Key == SEARCH_LOBBIES)
-PRAGMA_ENABLE_DEPRECATION_WARNINGS
+		if (Key == SEARCH_DEDICATED_ONLY || Key == SETTING_MAPNAME || Key == SEARCH_EMPTY_SERVERS_ONLY || Key == SEARCH_SECURE_SERVERS_ONLY || Key == SEARCH_LOBBIES)
 		{
 			continue;
 		}
diff --git a/Engine/Plugins/Online/OnlineSubsystemSteam/Source/Private/OnlineSessionInterfaceSteam.cpp b/Engine/Plugins/Online/OnlineSubsystemSteam/Source/Private/OnlineSessionInterfaceSteam.cpp
--- a/Engine/Plugins/Online/OnlineSubsystemSteam/Source/Private/OnlineSessionInterfaceSteam.cpp
+++ b/Engine/Plugins/Online/OnlineSubsystemSteam/Source/Private/OnlineSessionInterfaceSteam.cpp
@@ -775,11 +775,7 @@ uint32 FOnlineSessionSteam::FindInternetSession(const TSharedRef<FOnlineSessionS
 {
 	bool PresenceSearch = false;
 	bool LobbySearch = false;
-	if (
-PRAGMA_DISABLE_DEPRECATION_WARNINGS
-		(SearchSettings->QuerySettings.Get(SEARCH_PRESENCE, PresenceSearch) && PresenceSearch)
-PRAGMA_ENABLE_DEPRECATION_WARNINGS
-		|| (SearchSettings->QuerySettings.Get(SEARCH_LOBBIES, LobbySearch) && LobbySearch))
+	if (SearchSettings->QuerySettings.Get(SEARCH_LOBBIES, LobbySearch) && LobbySearch)
 	{
 		FOnlineAsyncTaskSteamFindLobbies* NewTask = new FOnlineAsyncTaskSteamFindLobbies(SteamSubsystem, SearchSettings);
 		SteamSubsystem->QueueAsyncTask(NewTask);
diff --git a/Engine/Plugins/Online/OnlineSubsystemUtils/Source/OnlineSubsystemUtils/Classes/AchievementWriteCallbackProxy.h b/Engine/Plugins/Online/OnlineSubsystemUtils/Source/OnlineSubsystemUtils/Classes/AchievementWriteCallbackProxy.h
--- a/Engine/Plugins/Online/OnlineSubsystemUtils/Source/OnlineSubsystemUtils/Classes/AchievementWriteCallbackProxy.h
+++ b/Engine/Plugins/Online/OnlineSubsystemUtils/Source/OnlineSubsystemUtils/Classes/AchievementWriteCallbackProxy.h
@@ -23,22 +23,6 @@ class UAchievementWriteCallbackProxy : public UOnlineBlueprintCallProxyBase
 	UPROPERTY(BlueprintAssignable)
 	FAchievementWriteCompleteDelegate OnWriteFailure;
 
-	UE_DEPRECATED(5.5, "Use OnWriteSuccess instead.")
-	UPROPERTY(BlueprintAssignable, meta = (DeprecatedProperty, DeprecationMessage = "Use OnWriteSuccess instead."))
-	FAchievementWriteDelegate OnSuccess;
-
-	UE_DEPRECATED(5.5, "Use OnWriteFailure instead.")
-	UPROPERTY(BlueprintAssignable, meta = (DeprecatedProperty, DeprecationMessage = "Use OnWriteFailure instead."))
-	FAchievementWriteDelegate OnFailure;
-
-	// Writes progress about an achievement to the default online subsystem
-	//   AchievementName is the ID of the achievement to update progress on
-	//   Progress is the reported progress toward accomplishing the achievement
-	//   UserTag is not used internally, but it is returned on success or failure
-	UE_DEPRECATED(5.5, "Use WriteProgress instead.")
-	UFUNCTION(BlueprintCallable, meta=(BlueprintInternalUseOnly = "true", WorldContext="WorldContextObject", DeprecatedFunction, DeprecationMessage = "Use WriteProgress instead."), Category="Online|Achievements")
-	static UAchievementWriteCallbackProxy* WriteAchievementProgress(UObject* WorldContextObject, APlayerController* PlayerController, FName AchievementName, float Progress = 100.0f, int32 UserTag = 0) { return WriteProgress(WorldContextObject, PlayerController, AchievementName.ToString(), Progress, UserTag); }
-
 	UFUNCTION(BlueprintCallable, meta=(BlueprintInternalUseOnly = "true", WorldContext="WorldContextObject"), Category="Online|Achievements")
 	static UAchievementWriteCallbackProxy* WriteProgress(UObject* WorldContextObject, APlayerController* PlayerController, FString AchievementName, float Progress = 100.0f, int32 UserTag = 0);
 
diff --git a/Engine/Plugins/Online/OnlineSubsystemUtils/Source/OnlineSubsystemUtils/Private/AchievementWriteCallbackProxy.cpp b/Engine/Plugins/Online/OnlineSubsystemUtils/Source/OnlineSubsystemUtils/Private/AchievementWriteCallbackProxy.cpp
--- a/Engine/Plugins/Online/OnlineSubsystemUtils/Source/OnlineSubsystemUtils/Private/AchievementWriteCallbackProxy.cpp
+++ b/Engine/Plugins/Online/OnlineSubsystemUtils/Source/OnlineSubsystemUtils/Private/AchievementWriteCallbackProxy.cpp
@@ -58,10 +58,6 @@ void UAchievementWriteCallbackProxy::Activate()
 	// Fail immediately
 	OnWriteFailure.Broadcast(AchievementName, AchievementProgress, UserTag);
 	
-PRAGMA_DISABLE_DEPRECATION_WARNINGS
-	OnFailure.Broadcast(FName(AchievementName), AchievementProgress, UserTag);
-PRAGMA_ENABLE_DEPRECATION_WARNINGS
-
 	WriteObject.Reset();
 }
 
@@ -70,18 +66,10 @@ void UAchievementWriteCallbackProxy::OnAchievementWritten(const FUniqueNetId& Us
 	if (bSuccess)
 	{
 		OnWriteSuccess.Broadcast(AchievementName, AchievementProgress, UserTag);
-
-PRAGMA_DISABLE_DEPRECATION_WARNINGS
-		OnSuccess.Broadcast(FName(AchievementName), AchievementProgress, UserTag);	
-PRAGMA_ENABLE_DEPRECATION_WARNINGS
 	}
 	else
 	{
 		OnWriteFailure.Broadcast(AchievementName, AchievementProgress, UserTag);
-		
-PRAGMA_DISABLE_DEPRECATION_WARNINGS
-		OnFailure.Broadcast(FName(AchievementName), AchievementProgress, UserTag);	
-PRAGMA_ENABLE_DEPRECATION_WARNINGS
 	}
 
 	WriteObject.Reset();
diff --git a/Engine/Source/Runtime/Online/BackgroundHTTP/Private/IOS/ApplePlatformBackgroundHttpRequest.cpp b/Engine/Source/Runtime/Online/BackgroundHTTP/Private/IOS/ApplePlatformBackgroundHttpRequest.cpp
--- a/Engine/Source/Runtime/Online/BackgroundHTTP/Private/IOS/ApplePlatformBackgroundHttpRequest.cpp
+++ b/Engine/Source/Runtime/Online/BackgroundHTTP/Private/IOS/ApplePlatformBackgroundHttpRequest.cpp
@@ -29,12 +29,6 @@ void FApplePlatformBackgroundHttpRequest::GetDebugText(TArray<FString>& Output)
 }
 #endif
 
-const FString& FApplePlatformBackgroundHttpRequest::GetRequestDebugID() const
-{
-	static const FString EmptyString = TEXT("");
-	return EmptyString;
-}
-
 void FApplePlatformBackgroundHttpRequest::UpdateProgress()
 {
 	if (DownloadId == FBackgroundURLSessionHandler::InvalidDownloadId)
diff --git a/Engine/Source/Runtime/Online/BackgroundHTTP/Public/IOS/ApplePlatformBackgroundHttpManager.h b/Engine/Source/Runtime/Online/BackgroundHTTP/Public/IOS/ApplePlatformBackgroundHttpManager.h
--- a/Engine/Source/Runtime/Online/BackgroundHTTP/Public/IOS/ApplePlatformBackgroundHttpManager.h
+++ b/Engine/Source/Runtime/Online/BackgroundHTTP/Public/IOS/ApplePlatformBackgroundHttpManager.h
@@ -28,9 +28,6 @@ public:
 	virtual bool IsGenericImplementation() const override { return false; }
 	virtual bool Tick(float DeltaTime) override;
 
-	UE_DEPRECATED(5.5, "Variable deprecated") static FString BackgroundSessionIdentifier;
-	UE_DEPRECATED(5.5, "Variable deprecated") static float ActiveTimeOutSetting;
-	UE_DEPRECATED(5.5, "Variable deprecated") static int RetryResumeDataLimitSetting;
 private:
 	TArray<FBackgroundHttpRequestPtr> PendingRemoveRequests;
 	FRWLock PendingRemoveRequestLock;
diff --git a/Engine/Source/Runtime/Online/BackgroundHTTP/Public/IOS/ApplePlatformBackgroundHttpRequest.h b/Engine/Source/Runtime/Online/BackgroundHTTP/Public/IOS/ApplePlatformBackgroundHttpRequest.h
--- a/Engine/Source/Runtime/Online/BackgroundHTTP/Public/IOS/ApplePlatformBackgroundHttpRequest.h
+++ b/Engine/Source/Runtime/Online/BackgroundHTTP/Public/IOS/ApplePlatformBackgroundHttpRequest.h
@@ -18,11 +18,6 @@ public:
 	virtual void PauseRequest() override;
 	virtual void ResumeRequest() override;
 
-	//Used to provide some extra debug information over normal GetRequestID()
-	//Returns string in format of X.Y where X is the underlying Task Identifier if set and Y is what was set in the SetRequestID() call if this reqeust has associated with a task.
-	//Returns the same as GetRequestID() if no task has been associated yet.
-	UE_DEPRECATED(5.5, "Method deprecated") const FString& GetRequestDebugID() const;
-
 #if !UE_BUILD_SHIPPING
 	virtual void GetDebugText(TArray<FString>& Output) override;
 #endif
diff --git a/Engine/Source/Runtime/Online/HTTP/Private/GenericPlatform/GenericPlatformHttp.cpp b/Engine/Source/Runtime/Online/HTTP/Private/GenericPlatform/GenericPlatformHttp.cpp
--- a/Engine/Source/Runtime/Online/HTTP/Private/GenericPlatform/GenericPlatformHttp.cpp
+++ b/Engine/Source/Runtime/Online/HTTP/Private/GenericPlatform/GenericPlatformHttp.cpp
@@ -321,13 +321,6 @@ IHttpRequest* FGenericPlatformHttp::ConstructRequest()
 	return new FGenericPlatformHttpRequest();
 }
 
-PRAGMA_DISABLE_DEPRECATION_WARNINGS
-bool FGenericPlatformHttp::UsesThreadedHttp()
-{
-	return true;
-}
-PRAGMA_ENABLE_DEPRECATION_WARNINGS
-
 static bool IsAllowedChar(UTF8CHAR LookupChar)
 {
 	static char AllowedChars[] = "ABCDEFGHIJKLMNOPQRSTUVWXYZabcdefghijklmnopqrstuvwxyz0123456789-_.~";
diff --git a/Engine/Source/Runtime/Online/HTTP/Private/HttpManager.cpp b/Engine/Source/Runtime/Online/HTTP/Private/HttpManager.cpp
--- a/Engine/Source/Runtime/Online/HTTP/Private/HttpManager.cpp
+++ b/Engine/Source/Runtime/Online/HTTP/Private/HttpManager.cpp
@@ -60,9 +60,6 @@ const TCHAR* LexToString(const EHttpFlushReason& FlushReason)
 	switch (FlushReason)
 	{
 	case EHttpFlushReason::Default:		return TEXT("Default");
-PRAGMA_DISABLE_DEPRECATION_WARNINGS
-	case EHttpFlushReason::Background:	return TEXT("Background");
-PRAGMA_ENABLE_DEPRECATION_WARNINGS
 	case EHttpFlushReason::Shutdown:	return TEXT("Shutdown");
 	case EHttpFlushReason::FullFlush:	return TEXT("FullFlush");
 	}
diff --git a/Engine/Source/Runtime/Online/HTTP/Private/Interfaces/IHttpRequest.cpp b/Engine/Source/Runtime/Online/HTTP/Private/Interfaces/IHttpRequest.cpp
--- a/Engine/Source/Runtime/Online/HTTP/Private/Interfaces/IHttpRequest.cpp
+++ b/Engine/Source/Runtime/Online/HTTP/Private/Interfaces/IHttpRequest.cpp
@@ -24,19 +24,6 @@ private:
 	FHttpRequestStreamDelegateV2 StreamDelegateV2;
 };
 
-
-PRAGMA_DISABLE_DEPRECATION_WARNINGS 
-bool IHttpRequest::SetContentFromStreamDelegate(FHttpRequestStreamDelegate StreamDelegate) 
-{ 
-	return SetContentFromStream(MakeShared<FArchiveWithDelegate>(StreamDelegate)); 
-}
-
-bool IHttpRequest::SetResponseBodyReceiveStreamDelegate(FHttpRequestStreamDelegate StreamDelegate) 
-{ 
-	return SetResponseBodyReceiveStream(MakeShared<FArchiveWithDelegate>(StreamDelegate)); 
-}
-PRAGMA_ENABLE_DEPRECATION_WARNINGS 
-
 bool IHttpRequest::SetResponseBodyReceiveStreamDelegateV2(FHttpRequestStreamDelegateV2 StreamDelegate) 
 { 
 	return SetResponseBodyReceiveStream(MakeShared<FArchiveWithDelegateV2>(StreamDelegate)); 
diff --git a/Engine/Source/Runtime/Online/HTTP/Public/GenericPlatform/GenericPlatformHttp.h b/Engine/Source/Runtime/Online/HTTP/Public/GenericPlatform/GenericPlatformHttp.h
--- a/Engine/Source/Runtime/Online/HTTP/Public/GenericPlatform/GenericPlatformHttp.h
+++ b/Engine/Source/Runtime/Online/HTTP/Public/GenericPlatform/GenericPlatformHttp.h
@@ -85,14 +85,6 @@ public:
 	 */
 	static HTTP_API IHttpRequest* ConstructRequest();
 
-	/**
-	 * Check if a platform uses the HTTP thread
-	 *
-	 * @return true if the platform uses threaded HTTP, false if not
-	 */
-	UE_DEPRECATED(5.5, "UsesThreadedHttp is deprecated and will be removed")
-	static HTTP_API bool UsesThreadedHttp();
-
 	/**
 	 * Returns a percent-encoded version of the passed in string
 	 *
diff --git a/Engine/Source/Runtime/Online/HTTP/Public/HttpManager.h b/Engine/Source/Runtime/Online/HTTP/Public/HttpManager.h
--- a/Engine/Source/Runtime/Online/HTTP/Public/HttpManager.h
+++ b/Engine/Source/Runtime/Online/HTTP/Public/HttpManager.h
@@ -28,8 +28,6 @@ enum class EHttpFlushReason : uint8
 	/** Reasonable, typically higher time limits */
 	Default,
 	/** Shorter time limits depending on platform requirements */
-	Background UE_DEPRECATED(5.5, "Flush when go into background is no longer used and will be removed"),
-	/** Shorter time limits depending on platform requirements */
 	Shutdown,
 	/** Infinite wait, should only be used in non-game scenarios where longer waits are acceptable */
 	FullFlush,
@@ -90,20 +88,6 @@ struct FHttpStats
 	/** The optional http stats on specific platform */
 	TOptional<FHttpStatsPlatform> PlatformStats;
 
-	UE_DEPRECATED(5.5, "operator== for FHttpStats has been deprecated and will be removed.")
-	bool operator==(const FHttpStats& Other) const
-	{
-		return RequestsInQueue == Other.RequestsInQueue
-			&& RequestsInFlight == Other.RequestsInFlight
-			&& MaxRequestsInQueue == Other.MaxRequestsInQueue
-			&& MaxRequestsInFlight == Other.MaxRequestsInFlight
-			&& TotalDownloadedBytes == Other.TotalDownloadedBytes
-			&& BandwidthMbps == Other.BandwidthMbps
-			&& HttpDurationMsAvg == Other.HttpDurationMsAvg
-			&& FMath::IsNearlyEqual(MaxTimeToConnect, Other.MaxTimeToConnect)
-			&& FMath::IsNearlyEqual(MaxTimeToWaitInQueue, Other.MaxTimeToWaitInQueue);
-	}
-
 	FHttpStats()
 	{
 	}
diff --git a/Engine/Source/Runtime/Online/HTTP/Public/Interfaces/IHttpRequest.h b/Engine/Source/Runtime/Online/HTTP/Public/Interfaces/IHttpRequest.h
--- a/Engine/Source/Runtime/Online/HTTP/Public/Interfaces/IHttpRequest.h
+++ b/Engine/Source/Runtime/Online/HTTP/Public/Interfaces/IHttpRequest.h
@@ -228,19 +228,6 @@ public:
 	 */
 	virtual bool SetContentFromStream(TSharedRef<FArchive, ESPMode::ThreadSafe> Stream) = 0;
 
-	/**
-	 * Sets the content of the request to stream directly from an delegate.
-	 * NOTE: 
-	 *   - The delegate will be called from another thread other than the game thread, make sure 
-	 *     it's thread-safe in there
-	 *   - Make sure the delegate is safe to be called until receiving the process complete callback
-	 *     or after canceling the request
-	 * @param StreamDelegate - delegate from which the payload should be streamed.
-	 * @return True if the delegate can be used to stream the request. False otherwise.
-	 */
-	UE_DEPRECATED(5.5, "SetContentFromStreamDelegate has been deprecated and will not be supported because there is no seek support through delegate. Implement your own FArchive instead.")
-	bool SetContentFromStreamDelegate(FHttpRequestStreamDelegate StreamDelegate);
-
 	/**
 	 * Sets the stream to receive the response body. Make sure to handle the cleanup of stream when
 	 * Serialize generated error(Stream->GetError returns true after Stream->Serialize call), this
@@ -255,20 +242,6 @@ public:
 	 */
 	virtual bool SetResponseBodyReceiveStream(TSharedRef<FArchive> Stream) = 0;
 
-	/**
-	 * Sets the delegate to receive the response body. Make sure to handle the cleanup of received data when
-	 * failed to process the data(StreamDelegate return false), this http request will fail and quit.
-	 *
-	 * NOTE: Once set, the data will no longer be cached in response, IHttpResponse::GetContent() and
-	 * IHttpResponse::GetContentAsString() will return empty result. The delegate will be called from
-	 * another thread other than the game thread
-	 *
-	 * @param StreamDelegate - will be used to receive the response body
-	 * @return True if the delegate can be used. False otherwise.
-	 */
-	UE_DEPRECATED(5.5, "SetResponseBodyReceiveStreamDelegate has been deprecated, use SetResponseBodyReceiveStreamDelegateV2 instead")
-	HTTP_API bool SetResponseBodyReceiveStreamDelegate(FHttpRequestStreamDelegate StreamDelegate);
-
 	/**
 	 * Sets the delegate to receive the response body. Make sure to handle the cleanup of received data when
 	 * failed to process the data(StreamDelegate return false), this http request will fail and quit.
