Added OSSONEStore plugin.

diff --git a/Engine/Plugins/Online/OnlineSubsystem/Source/Private/OnlineSubsystem.cpp b/Engine/Plugins/Online/OnlineSubsystem/Source/Private/OnlineSubsystem.cpp
--- a/Engine/Plugins/Online/OnlineSubsystem/Source/Private/OnlineSubsystem.cpp
+++ b/Engine/Plugins/Online/OnlineSubsystem/Source/Private/OnlineSubsystem.cpp
@@ -132,6 +132,10 @@ const FName NINTENDO_SUBSYSTEM(TEXT("NINTENDO"));
 const FName SAMSUNG_SUBSYSTEM(TEXT("SAMSUNG"));
 #endif
 
+#ifndef ONESTORE_SUBSYSTEM
+const FName ONESTORE_SUBSYSTEM(TEXT("ONESTORE"));
+#endif
+
 #ifndef EOS_SUBSYSTEM
 const FName EOS_SUBSYSTEM(TEXT("EOS"));
 #endif
diff --git a/Engine/Plugins/Online/OnlineSubsystem/Source/Public/OnlineSubsystemNames.h b/Engine/Plugins/Online/OnlineSubsystem/Source/Public/OnlineSubsystemNames.h
--- a/Engine/Plugins/Online/OnlineSubsystem/Source/Public/OnlineSubsystemNames.h
+++ b/Engine/Plugins/Online/OnlineSubsystem/Source/Public/OnlineSubsystemNames.h
@@ -90,6 +90,10 @@ extern ONLINESUBSYSTEM_API const FName NINTENDO_SUBSYSTEM;
 extern ONLINESUBSYSTEM_API const FName SAMSUNG_SUBSYSTEM;
 #endif
 
+#ifndef ONESTORE_SUBSYSTEM
+extern ONLINESUBSYSTEM_API const FName ONESTORE_SUBSYSTEM;
+#endif
+
 #ifndef EOS_SUBSYSTEM
 extern ONLINESUBSYSTEM_API const FName EOS_SUBSYSTEM;
 #endif
