OnlinePresenceInterface status from/to string -> .cpp file

Code size. Not a thing that needs to be inlined with its data replicated.

diff --git a/Engine/Plugins/Online/OnlineSubsystem/Source/Private/OnlinePresenceInterface.cpp b/Engine/Plugins/Online/OnlineSubsystem/Source/Private/OnlinePresenceInterface.cpp
--- a/Engine/Plugins/Online/OnlineSubsystem/Source/Private/OnlinePresenceInterface.cpp
+++ b/Engine/Plugins/Online/OnlineSubsystem/Source/Private/OnlinePresenceInterface.cpp
@@ -2,6 +2,93 @@
 
 #include "Interfaces/OnlinePresenceInterface.h"
 
+namespace EOnlinePresenceState
+{
+	const TCHAR* ToString(EOnlinePresenceState::Type EnumVal)
+	{
+		switch (EnumVal)
+		{
+		case Online:
+			return TEXT("Online");
+		case Offline:
+			return TEXT("Offline");
+		case Away:
+			return TEXT("Away");
+		case ExtendedAway:
+			return TEXT("ExtendedAway");
+		case DoNotDisturb:
+			return TEXT("DoNotDisturb");
+		case Chat:
+			return TEXT("Chat");
+		}
+		return TEXT("");
+	}
+
+	EOnlinePresenceState::Type FromString(const TCHAR* StringVal)
+	{
+		if (FCString::Stricmp(StringVal, TEXT("Online")) == 0)
+		{
+			return EOnlinePresenceState::Online;
+		}
+		else if (FCString::Stricmp(StringVal, TEXT("Offline")) == 0)
+		{
+			return EOnlinePresenceState::Offline;
+		}
+		else if (FCString::Stricmp(StringVal, TEXT("Away")) == 0)
+		{
+			return EOnlinePresenceState::Away;
+		}
+		else if (FCString::Stricmp(StringVal, TEXT("ExtendedAway")) == 0)
+		{
+			return EOnlinePresenceState::ExtendedAway;
+		}
+		else if (FCString::Stricmp(StringVal, TEXT("DoNotDisturb")) == 0)
+		{
+			return EOnlinePresenceState::DoNotDisturb;
+		}
+		else if (FCString::Stricmp(StringVal, TEXT("Chat")) == 0)
+		{
+			return EOnlinePresenceState::Chat;
+		}
+		// Default to Offline / generally unavailable
+		return EOnlinePresenceState::Offline;
+	}
+
+	// Move these into statics in ToLocText once deprecation period expires
+PRAGMA_DISABLE_DEPRECATION_WARNINGS
+	const FText OnlineText =  NSLOCTEXT("OnlinePresence", "Online", "Online");
+	const FText OfflineText =  NSLOCTEXT("OnlinePresence", "Offline", "Offline");
+	const FText AwayText =  NSLOCTEXT("OnlinePresence", "Away", "Away");
+	const FText DoNotDisturbText =  NSLOCTEXT("OnlinePresence", "DoNotDisturb", "Do Not Disturb");
+	const FText ChatText =  NSLOCTEXT("OnlinePresence", "Chat", "Chat");
+PRAGMA_ENABLE_DEPRECATION_WARNINGS
+
+	/** 
+	 * @return the loc text version of the enum passed in 
+	 */
+	const FText ToLocText(EOnlinePresenceState::Type EnumVal)
+	{
+PRAGMA_DISABLE_DEPRECATION_WARNINGS
+		switch (EnumVal)
+		{
+		case Online:
+			return OnlineText;
+		case Offline:
+			return OfflineText;
+		case ExtendedAway:
+			// falls through to return away text
+		case Away:
+			return AwayText;
+		case DoNotDisturb:
+			return DoNotDisturbText;
+		case Chat:
+			return ChatText;
+		}
+PRAGMA_ENABLE_DEPRECATION_WARNINGS
+		return FText::GetEmpty();
+	}
+}
+
 // Implementation for OnlineSubsystems that have not implemented this version yet
 void IOnlinePresence::SetPresence(const FUniqueNetId& User, FOnlinePresenceSetPresenceParameters&& Parameters, const FOnPresenceTaskCompleteDelegate& Delegate)
 {
diff --git a/Engine/Plugins/Online/OnlineSubsystem/Source/Public/Interfaces/OnlinePresenceInterface.h b/Engine/Plugins/Online/OnlineSubsystem/Source/Public/Interfaces/OnlinePresenceInterface.h
--- a/Engine/Plugins/Online/OnlineSubsystem/Source/Public/Interfaces/OnlinePresenceInterface.h
+++ b/Engine/Plugins/Online/OnlineSubsystem/Source/Public/Interfaces/OnlinePresenceInterface.h
@@ -67,87 +67,28 @@ namespace EOnlinePresenceState
 	/** 
 	 * @return the stringified version of the enum passed in 
 	 */
-	inline const TCHAR* ToString(EOnlinePresenceState::Type EnumVal)
-	{
-		switch (EnumVal)
-		{
-		case Online:
-			return TEXT("Online");
-		case Offline:
-			return TEXT("Offline");
-		case Away:
-			return TEXT("Away");
-		case ExtendedAway:
-			return TEXT("ExtendedAway");
-		case DoNotDisturb:
-			return TEXT("DoNotDisturb");
-		case Chat:
-			return TEXT("Chat");
-		}
-		return TEXT("");
-	}
+	ONLINESUBSYSTEM_API const TCHAR* ToString(EOnlinePresenceState::Type EnumVal);
 
 	/**
 	 * @return EOnlinePresenceState from the string passed in
 	 */
-	inline EOnlinePresenceState::Type FromString(const TCHAR* StringVal)
-	{
-		if (FCString::Stricmp(StringVal, TEXT("Online")) == 0)
-		{
-			return EOnlinePresenceState::Online;
-		}
-		else if (FCString::Stricmp(StringVal, TEXT("Offline")) == 0)
-		{
-			return EOnlinePresenceState::Offline;
-		}
-		else if (FCString::Stricmp(StringVal, TEXT("Away")) == 0)
-		{
-			return EOnlinePresenceState::Away;
-		}
-		else if (FCString::Stricmp(StringVal, TEXT("ExtendedAway")) == 0)
-		{
-			return EOnlinePresenceState::ExtendedAway;
-		}
-		else if (FCString::Stricmp(StringVal, TEXT("DoNotDisturb")) == 0)
-		{
-			return EOnlinePresenceState::DoNotDisturb;
-		}
-		else if (FCString::Stricmp(StringVal, TEXT("Chat")) == 0)
-		{
-			return EOnlinePresenceState::Chat;
-		}
-		// Default to Offline / generally unavailable
-		return EOnlinePresenceState::Offline;
-	}
+	ONLINESUBSYSTEM_API EOnlinePresenceState::Type FromString(const TCHAR* StringVal);
+
+	UE_DEPRECATED(5.8, "Use ToLocText(EOnlinePresenceState::Online)")
+	extern ONLINESUBSYSTEM_API const FText OnlineText;
+	UE_DEPRECATED(5.8, "Use ToLocText(EOnlinePresenceState::Offline)")
+	extern ONLINESUBSYSTEM_API const FText OfflineText;
+	UE_DEPRECATED(5.8, "Use ToLocText(EOnlinePresenceState::Away)")
+	extern ONLINESUBSYSTEM_API const FText AwayText;
+	UE_DEPRECATED(5.8, "Use ToLocText(EOnlinePresenceState::DoNotDisturb)")
+	extern ONLINESUBSYSTEM_API const FText DoNotDisturbText;
+	UE_DEPRECATED(5.8, "Use ToLocText(EOnlinePresenceState::Chat)")
+	extern ONLINESUBSYSTEM_API const FText ChatText;
 
-	static FText OnlineText =  NSLOCTEXT("OnlinePresence", "Online", "Online");
-	static FText OfflineText =  NSLOCTEXT("OnlinePresence", "Offline", "Offline");
-	static FText AwayText =  NSLOCTEXT("OnlinePresence", "Away", "Away");
-	static FText DoNotDisturbText =  NSLOCTEXT("OnlinePresence", "DoNotDisturb", "Do Not Disturb");
-	static FText ChatText =  NSLOCTEXT("OnlinePresence", "Chat", "Chat");
 	/** 
 	 * @return the loc text version of the enum passed in 
 	 */
-	inline const FText ToLocText(EOnlinePresenceState::Type EnumVal)
-	{
-		switch (EnumVal)
-		{
-		case Online:
-			return OnlineText;
-		case Offline:
-			return OfflineText;
-		case ExtendedAway:
-			// falls through to return away text
-		case Away:
-			return AwayText;
-		case DoNotDisturb:
-			return DoNotDisturbText;
-		case Chat:
-			return ChatText;
-		}
-		return FText::GetEmpty();
-	}
-
+	ONLINESUBSYSTEM_API const FText ToLocText(EOnlinePresenceState::Type EnumVal);
 }
 
 class FOnlineUserPresenceStatus
