Catch store ID missmatch in GDK OSS

part2

diff --git a/Engine/Plugins/Online/OnlineSubsystem/Source/Private/OnlineError.cpp b/Engine/Plugins/Online/OnlineSubsystem/Source/Private/OnlineError.cpp
--- a/Engine/Plugins/Online/OnlineSubsystem/Source/Private/OnlineError.cpp
+++ b/Engine/Plugins/Online/OnlineSubsystem/Source/Private/OnlineError.cpp
@@ -199,6 +199,8 @@ FString FOnlineError::DefaultErrorCode(EOnlineErrorResult Result)
 			return FString(TEXT("canceled"));
 		case EOnlineErrorResult::FailExtended:
 			return FString(TEXT("fail_extended"));
+		case EOnlineErrorResult::MismatchedUser:
+			return FString(TEXT("mismatched_user"));
 	};
 	return FString(TEXT("unknown_error_result"));
 }
@@ -243,6 +245,8 @@ FText FOnlineError::DefaultErrorMsg(EOnlineErrorResult Result)
 			return LOCTEXT("Canceled", "Operation was canceled");
 		case EOnlineErrorResult::FailExtended:
 			return LOCTEXT("FailExtended", "Extended error");
+		case EOnlineErrorResult::MismatchedUser:
+			return LOCTEXT("MismatchedUser", "Title/Store ID Mismatch");
 	};
 	return LOCTEXT("Unknown", "Unknown error");
 }
diff --git a/Engine/Plugins/Online/OnlineSubsystem/Source/Public/OnlineError.h b/Engine/Plugins/Online/OnlineSubsystem/Source/Public/OnlineError.h
--- a/Engine/Plugins/Online/OnlineSubsystem/Source/Public/OnlineError.h
+++ b/Engine/Plugins/Online/OnlineSubsystem/Source/Public/OnlineError.h
@@ -56,6 +56,8 @@ enum class EOnlineErrorResult : uint8
 	FailExtended,
 	/** No game session found */
 	NoGameSession,
+	/** Failed due to mismatched missing user between title and store */
+	MismatchedUser,
 
 	/** Default state */
 	Unknown
