Cache GConfig read in IOnlineSubsystem::GetLocalPlatformName to reduce potential hitches if GEngineIni gets unloaded
#Synced-CL 49624953

diff --git a/Engine/Plugins/Online/OnlineSubsystem/Source/Private/OnlineSubsystem.cpp b/Engine/Plugins/Online/OnlineSubsystem/Source/Private/OnlineSubsystem.cpp
--- a/Engine/Plugins/Online/OnlineSubsystem/Source/Private/OnlineSubsystem.cpp
+++ b/Engine/Plugins/Online/OnlineSubsystem/Source/Private/OnlineSubsystem.cpp
@@ -287,7 +287,16 @@ FString IOnlineSubsystem::GetLocalPlatformName()
 #endif
 	if (OnlinePlatform.IsEmpty())
 	{
-		GConfig->GetString(TEXT("OnlineSubsystem"), TEXT("LocalPlatformName"), OnlinePlatform, GEngineIni);
+		// Only read from config once so we don't re-load ini files if they get unloaded. This value is expected to never change. Use the cvar if you need dynamic changing.
+		static bool bReadFromConfig = false;
+		static FString ConfigValue;
+		if (!bReadFromConfig)
+		{
+			GConfig->GetString(TEXT("OnlineSubsystem"), TEXT("LocalPlatformName"), ConfigValue, GEngineIni);
+			// Regardless of whether we found a string in config or not, do not read again
+			bReadFromConfig = true;
+		}
+		OnlinePlatform = ConfigValue;
 	}
 
 	if (!OnlinePlatform.IsEmpty())
