UE: documentation for OnlineSubsystem's ModulesToLoad.
Small corrections in module code.

#Synced-CL 49763427

diff --git a/Engine/Config/BaseEngine.ini b/Engine/Config/BaseEngine.ini
--- a/Engine/Config/BaseEngine.ini
+++ b/Engine/Config/BaseEngine.ini
@@ -2453,6 +2453,11 @@ MaxRemoteTalkers=16
 PollingIntervalInMs=20
 bUseBuildIdOverride=false
 BuildIdOverride=0
+
+; AdditionalModulesToLoad: Config driven array to load modules during this OnlineSubsystem's StartupModule().
+; Load modules that need to register callbacks or bind to SDK hooks before the module finishes initialization.
+; The Shutdown order of the Modules loaded here shut down before the parent module.
+; i.e.: OnlineSubsystem loads HTTP, shutdown is first HTTP and then OnlineSubsystem.
 !AdditionalModulesToLoad=Clear
 +AdditionalModulesToLoad=HTTP
 +AdditionalModulesToLoad=XMPP
diff --git a/Engine/Plugins/Online/OnlineSubsystem/Source/Private/OnlineSubsystemModule.cpp b/Engine/Plugins/Online/OnlineSubsystem/Source/Private/OnlineSubsystemModule.cpp
--- a/Engine/Plugins/Online/OnlineSubsystem/Source/Private/OnlineSubsystemModule.cpp
+++ b/Engine/Plugins/Online/OnlineSubsystem/Source/Private/OnlineSubsystemModule.cpp
@@ -77,11 +77,11 @@ void FOnlineSubsystemModule::StartupModule()
 	// We will always load these modules at the cost of extra modules loaded for the few OSS (like Null) that don't use it.
 	TArray<FString> AdditionalModulesToLoad;
 	GConfig->GetArray(TEXT("OnlineSubsystem"), TEXT("AdditionalModulesToLoad"), AdditionalModulesToLoad, GEngineIni);
-	for (const FString& AdditonalModule : AdditionalModulesToLoad)
+	for (const FString& AdditionalModule : AdditionalModulesToLoad)
 	{
-		if (FModuleManager::Get().ModuleExists(*AdditonalModule))
+		if (FModuleManager::Get().ModuleExists(*AdditionalModule))
 		{
-			FModuleManager::Get().LoadModule(*AdditonalModule);
+			FModuleManager::Get().LoadModule(*AdditionalModule);
 		}
 	}
 
@@ -476,10 +476,8 @@ void FOnlineSubsystemModule::DestroyOnlineSubsystem(const FName InSubsystemName)
 
 bool FOnlineSubsystemModule::DoesInstanceExist(const FName InSubsystemName) const
 {
-	bool bIsLoaded = false;
-
 	FName SubsystemName, InstanceName;
-	FName KeyName = ParseOnlineSubsystemName(InSubsystemName, SubsystemName, InstanceName);
+	const FName KeyName = ParseOnlineSubsystemName(InSubsystemName, SubsystemName, InstanceName);
 	if (!SubsystemName.IsNone())
 	{
 		IOnlineSubsystemPtr OnlineSubsystem;
