Cache Config and CommandLine values for IOnlineSubsystem::IsEnabled. Minimize string allocations, especially when the values are already cached.

#tests Editor
#rnx

#Synced-CL 51003583

diff --git a/Engine/Plugins/Online/OnlineSubsystem/Source/Private/OnlineSubsystem.cpp b/Engine/Plugins/Online/OnlineSubsystem/Source/Private/OnlineSubsystem.cpp
--- a/Engine/Plugins/Online/OnlineSubsystem/Source/Private/OnlineSubsystem.cpp
+++ b/Engine/Plugins/Online/OnlineSubsystem/Source/Private/OnlineSubsystem.cpp
@@ -3,6 +3,7 @@
 #include "OnlineSubsystem.h"
 #include "Misc/CommandLine.h"
 #include "Misc/ConfigCacheIni.h"
+#include "Misc/ScopeRWLock.h"
 #include "Misc/TrackedActivity.h"
 #include "HAL/IConsoleManager.h"
 #include "Online/OnlineSessionNames.h"
@@ -387,37 +388,109 @@ FAutoConsoleCommand CmdResetAchievements(
 
 #endif
 
-bool IOnlineSubsystem::IsEnabled(const FName& SubsystemName, const FName& InstanceName)
+
+bool CheckSubsystemDisabledByCommandLine(const FName& SubsystemName)
 {
-	bool bIsDisabledByCommandLine = false;
 #if (!UE_BUILD_SHIPPING && !UE_BUILD_SHIPPING_WITH_EDITOR) || ENABLE_PGO_PROFILE
 	// In non-shipping builds, or PGO_PROFILE enabled Shipping builds, check for a command line override to disable the OSS
-	bIsDisabledByCommandLine = FParse::Param(FCommandLine::Get(), *FString::Printf(TEXT("no%s"), *SubsystemName.ToString()));
+	// Cache the results in a static map to avoid string allocations after the first query
+	static TMap<FName, bool> DisableMap;
+	static FRWLock DisableMapLock;
+
+	// Check the cache first
+	{
+		UE::TReadScopeLock ReadScope(DisableMapLock);
+
+		bool* FoundValue = DisableMap.Find(SubsystemName);
+		if (FoundValue)
+		{
+			return *FoundValue;
+		}
+	}
+
+	// Check the actual command line, add to cache
+	bool bIsDisabledByCommandLine = FParse::Param(FCommandLine::Get(), *FString::Printf(TEXT("no%s"), *SubsystemName.ToString()));
+
+	{
+		UE::TWriteScopeLock WriteScope(DisableMapLock);
+		DisableMap.Add(SubsystemName, bIsDisabledByCommandLine);
+	}
+
+	return bIsDisabledByCommandLine;
+#else
+	return false;
 #endif
-	
-	if (!bIsDisabledByCommandLine)
+}
+
+bool CheckSubsystemEnabledByConfig(const FName& SubsystemName, const FName& InstanceName, bool& bOutEnabled)
+{
+	// Map of [SubsystemName, InstanceName] to bEnabled config value, unset if the value is not specified.
+	// Note that the subsystem is queried on it's own so there will always be a [SubsystemName, None] entry in the map
+	static TMap<TTuple<FName, FName>, TOptional<bool>> EnableMap;
+	static FRWLock EnableMapLock;
+
+	TTuple<FName, FName> MapKeyNames = MakeTuple(SubsystemName, InstanceName);
+
+	// Check the cache first
 	{
-		bool bIsEnabledByConfig = false;
-		bool bConfigOptionExists = false;
-		for (int32 InstancePass = 0; InstancePass < 2; InstancePass++)
+		UE::TReadScopeLock ReadScope(EnableMapLock);
+		TOptional<bool>* FoundEnableValue = EnableMap.Find(MapKeyNames);
+		if (FoundEnableValue)
 		{
-			FString InstanceSection;
-			if (InstancePass == 1)
+			bool bFoundValueIsSet = FoundEnableValue->IsSet();
+			if (bFoundValueIsSet)
 			{
-				if (InstanceName != NAME_None)
-				{
-					InstanceSection = TEXT(" ") + InstanceName.ToString();
-				}
-				else
-				{
-					continue;
-				}
+				bOutEnabled = FoundEnableValue->GetValue();
 			}
 
-			FString ConfigSection(FString::Printf(TEXT("OnlineSubsystem%s"), *SubsystemName.ToString()));
-			ConfigSection += InstanceSection;
-			bConfigOptionExists |= GConfig->GetBool(*ConfigSection, TEXT("bEnabled"), bIsEnabledByConfig, GEngineIni);
+			return bFoundValueIsSet;
+		}
+	}
+
+	// The section name is different if an InstanceName is specified
+	FString ConfigSection;
+	if (InstanceName == NAME_None)
+	{
+		ConfigSection = FString::Printf(TEXT("OnlineSubsystem%s"), *SubsystemName.ToString());
+	}
+	else
+	{
+		ConfigSection = FString::Printf(TEXT("OnlineSubsystem%s %s"), *SubsystemName.ToString(), *InstanceName.ToString());
+	}
+
+	bool bFoundConfigValue = GConfig->GetBool(*ConfigSection, TEXT("bEnabled"), bOutEnabled, GEngineIni);
+
+	// Only set the TOptional if a config value was found. Otherwise being unset means the config value was unspecified.
+	TOptional<bool> MapValue;
+	if (bFoundConfigValue)
+	{
+		MapValue = bOutEnabled;
+	}
+
+	// Write new value to the cache
+	{
+		UE::TWriteScopeLock WriteScope(EnableMapLock);
+		EnableMap.Add(MapKeyNames, MapValue);
+	}
+
+	return bFoundConfigValue;
+}
+
+bool IOnlineSubsystem::IsEnabled(const FName& SubsystemName, const FName& InstanceName)
+{
+	bool bIsDisabledByCommandLine = CheckSubsystemDisabledByCommandLine(SubsystemName);
+	if (!bIsDisabledByCommandLine)
+	{
+		// First check the subsystem by itself
+		bool bIsEnabledByConfig = false;
+		bool bConfigOptionExists = CheckSubsystemEnabledByConfig(SubsystemName, NAME_None, bIsEnabledByConfig);
+
+		// Then check with the instance name if one is specified
+		if (InstanceName != NAME_None)
+		{
+			bConfigOptionExists |= CheckSubsystemEnabledByConfig(SubsystemName, InstanceName, bIsEnabledByConfig);
 		}
+
 		UE_CLOG_ONLINE(!bConfigOptionExists, Verbose, TEXT("[OnlineSubsystem%s].bEnabled is not set, defaulting to true"), *SubsystemName.ToString());
 	
 		return !bConfigOptionExists || bIsEnabledByConfig;
